Turning on AudioMixer for Windows, PS4, and XBOX One

diff --git a/Engine/Config/Windows/WindowsEngine.ini b/Engine/Config/Windows/WindowsEngine.ini
--- a/Engine/Config/Windows/WindowsEngine.ini
+++ b/Engine/Config/Windows/WindowsEngine.ini
@@ -1,7 +1,7 @@
 [Audio]
-AudioDeviceModuleName=XAudio2
+;AudioDeviceModuleName=XAudio2
 ; Uncomment below (and comment above) to use multi-platform mixer module by default
-;AudioDeviceModuleName=AudioMixerXAudio2
+AudioDeviceModuleName=AudioMixerXAudio2
 
 ; Defining below allows switching to audio mixer using -audiomixer commandline
 AudioMixerModuleName=AudioMixerXAudio2
