Reverting changes to engine ini files in CL 4953149

diff --git a/Engine/Config/Windows/WindowsEngine.ini b/Engine/Config/Windows/WindowsEngine.ini
--- a/Engine/Config/Windows/WindowsEngine.ini
+++ b/Engine/Config/Windows/WindowsEngine.ini
@@ -1,8 +1,5 @@
 [Audio]
-;AudioDeviceModuleName=XAudio2
-; Uncomment below (and comment above) to use multi-platform mixer module by default
-AudioDeviceModuleName=AudioMixerXAudio2
-
+AudioDeviceModuleName=XAudio2
 ; Defining below allows switching to audio mixer using -audiomixer commandline
 AudioMixerModuleName=AudioMixerXAudio2
 
