Flip audiomixer on by default.

diff --git a/Engine/Config/IOS/IOSEngine.ini b/Engine/Config/IOS/IOSEngine.ini
--- a/Engine/Config/IOS/IOSEngine.ini
+++ b/Engine/Config/IOS/IOSEngine.ini
@@ -18,6 +18,7 @@ AudioDeviceModuleName=IOSAudio
 AudioMixerModuleName=AudioMixerAudioUnit
 
 ; Defines a platform-specific volume headroom (in dB) for audio to provide better platform consistency with respect to volume levels.
+UseAudioMixer=true
 PlatformHeadroomDB=-6
 
 [PlatformMemoryBuckets]
diff --git a/Engine/Config/Lumin/LuminEngine.ini b/Engine/Config/Lumin/LuminEngine.ini
--- a/Engine/Config/Lumin/LuminEngine.ini
+++ b/Engine/Config/Lumin/LuminEngine.ini
@@ -5,6 +5,7 @@ AudioDeviceModuleName=MagicLeapAudio
 AudioMixerModuleName=MagicLeapAudio
 ; Defines a platform-specific volume headroom (in dB) for audio to provide better platform consistency with respect to volume levels.
 PlatformHeadroomDB=-4
+UseAudioMixer=true
 
 ;Override Android settings back to default
 [LocalNotification]
diff --git a/Engine/Config/Mac/MacEngine.ini b/Engine/Config/Mac/MacEngine.ini
--- a/Engine/Config/Mac/MacEngine.ini
+++ b/Engine/Config/Mac/MacEngine.ini
@@ -6,6 +6,7 @@ AudioMixerModuleName=AudioMixerCoreAudio
 
 ; Defines a platform-specific volume headroom (in dB) for audio to provide better platform consistency with respect to volume levels.
 PlatformHeadroomDB=-6
+UseAudioMixer=true
 
 [SystemSettings]
 r.setres=1280x720
diff --git a/Engine/Config/Windows/WindowsEngine.ini b/Engine/Config/Windows/WindowsEngine.ini
--- a/Engine/Config/Windows/WindowsEngine.ini
+++ b/Engine/Config/Windows/WindowsEngine.ini
@@ -5,6 +5,7 @@ AudioMixerModuleName=AudioMixerXAudio2
 
 ; Defines a platform-specific volume headroom (in dB) for audio to provide better platform consistency with respect to volume levels.
 PlatformHeadroomDB=-3
+UseAudioMixer=true
 
 [TextureStreaming]
 ; PoolSizeVRAMPercentage is how much percentage of GPU Dedicated VRAM should be used as a TexturePool cache for streaming textures (0 - unlimited streaming)
