disable r.Vulkan.UploadCmdBufferSemaphore
fix a validation error on shutdown

diff --git a/Engine/Config/Unix/UnixEngine.ini b/Engine/Config/Unix/UnixEngine.ini
--- a/Engine/Config/Unix/UnixEngine.ini
+++ b/Engine/Config/Unix/UnixEngine.ini
@@ -6,4 +6,3 @@ AudioMixerModuleName=AudioMixerSDL
 PlatformRequiresDataCrypto=True
 
 [SystemSettings]
-r.Vulkan.UploadCmdBufferSemaphore=1
diff --git a/Engine/Config/Windows/WindowsEngine.ini b/Engine/Config/Windows/WindowsEngine.ini
--- a/Engine/Config/Windows/WindowsEngine.ini
+++ b/Engine/Config/Windows/WindowsEngine.ini
@@ -18,7 +18,6 @@ ConfigClass=/Engine/DeviceProfiles/SetDeviceProfileWindows.SetDeviceProfileWindo
 r.setres=1280x720
 framepro.ScopeMinTimeMicroseconds=10
 fx.NiagaraAllowRuntimeScalabilityChanges=1
-r.Vulkan.UploadCmdBufferSemaphore=1
 
 [PlatformCrypto]
 PlatformRequiresDataCrypto=True
