- Changed the File|Package menu for Windows to only show 32 or 64-bit options, based on the bTarget32Bit Windows project setting, so that it won't cook with Win64 but run with Win32

diff --git a/Engine/Config/Windows/DataDrivenPlatformInfo.ini b/Engine/Config/Windows/DataDrivenPlatformInfo.ini
--- a/Engine/Config/Windows/DataDrivenPlatformInfo.ini
+++ b/Engine/Config/Windows/DataDrivenPlatformInfo.ini
@@ -7,32 +7,11 @@ Freezing_MaxFieldAlignment32=4
 Freezing_MaxFieldAlignment64=8
 Freezing_bAlignBases=true
 
-[PlatformInfo WindowsNoEditor]
-TargetPlatformName=WindowsNoEditor
-DisplayName=Windows
-PlatformType=Game
-PlatformFlags=None
-NormalIconPath=Launcher/Windows/Platform_WindowsNoEditor_24x
-LargeIconPath=Launcher/Windows/Platform_WindowsNoEditor_128x
-XLargeIconPath=
-UATCommandLine=-targetplatform=Win64
-AutoSDKPath=
-TutorialPath=/Engine/Tutorial/Installation/InstallingVisualStudioTutorial.InstallingVisualStudioTutorial
-Windows:bIsEnabled=true
-Mac:bIsEnabled=false
-Linux:bIsEnabled=false
-BinariesDirectoryName=Win64
-IniPlatformName=Windows
-bUsesHostCompiler=true
-bUATClosesAfterLaunch=false
-UBTTargetID=Win64
-PlatformGroupName=Desktop
-
-[PlatformInfo WindowsNoEditor_Win32]
+[PlatformInfo WindowsNoEditorWin32]
 TargetPlatformName=WindowsNoEditor
 DisplayName=Windows (32-bit)
 PlatformType=Game
-PlatformFlags=BuildFlavor
+PlatformFlags=None
 NormalIconPath=Launcher/Windows/Platform_WindowsNoEditor_24x
 LargeIconPath=Launcher/Windows/Platform_WindowsNoEditor_128x
 XLargeIconPath=
@@ -49,11 +28,11 @@ bUATClosesAfterLaunch=false
 UBTTargetID=Win32
 PlatformGroupName=Desktop
 
-[PlatformInfo WindowsNoEditor_Win64]
+[PlatformInfo WindowsNoEditor]
 TargetPlatformName=WindowsNoEditor
 DisplayName=Windows (64-bit)
 PlatformType=Game
-PlatformFlags=BuildFlavor
+PlatformFlags=None
 NormalIconPath=Launcher/Windows/Platform_WindowsNoEditor_24x
 LargeIconPath=Launcher/Windows/Platform_WindowsNoEditor_128x
 XLargeIconPath=
diff --git a/Engine/Source/Editor/MainFrame/Private/Menus/PackageProjectMenu.h b/Engine/Source/Editor/MainFrame/Private/Menus/PackageProjectMenu.h
--- a/Engine/Source/Editor/MainFrame/Private/Menus/PackageProjectMenu.h
+++ b/Engine/Source/Editor/MainFrame/Private/Menus/PackageProjectMenu.h
@@ -68,6 +68,15 @@ public:
 				continue;
 			}
 
+			// windows target version needs special handling
+			bool bTargetWin32 = false;
+			GConfig->GetBool(TEXT("/Script/WindowsTargetPlatform.WindowsTargetSettings"), TEXT("bTarget32Bit"), bTargetWin32, GEngineIni);
+			// if we are targeting 32-bit, only show 32-bit option, for Windows
+			if ((VanillaPlatform.PlatformInfo->UBTTargetId == TEXT("Win32") && !bTargetWin32) || (VanillaPlatform.PlatformInfo->UBTTargetId == TEXT("Win64") && bTargetWin32))
+			{
+				continue;
+			}
+
 			// Check if this platform has a submenu entry
 			if (VanillaPlatform.PlatformInfo->PlatformSubMenu != NAME_None)
 			{
