- Added missing UBTPlatformName entry in Windows DDPI.ini

diff --git a/Engine/Config/Windows/DataDrivenPlatformInfo.ini b/Engine/Config/Windows/DataDrivenPlatformInfo.ini
--- a/Engine/Config/Windows/DataDrivenPlatformInfo.ini
+++ b/Engine/Config/Windows/DataDrivenPlatformInfo.ini
@@ -19,6 +19,7 @@ BinariesDirectoryName=Win64
 bUsesHostCompiler=true
 bUATClosesAfterLaunch=false
 PlatformGroupName=Desktop
+UBTPlatformName=Win64
 
 
 [PlatformInfo WindowsNoEditor]
@@ -54,4 +55,4 @@ bSupportsRayTracing=true
 bSupportsRayTracingMissShaderBindings=true
 bSupportsRayTracingIndirectInstanceData=true
 bSupportsWaveOperations=false
-bSupportsGPUScene=true
+bSupportsGPUScene=true
\ No newline at end of file
