UE5 - Add Vulkan SM5 to DDPI
Removed duplicate assignment
#rnx

diff --git a/Engine/Config/Linux/DataDrivenPlatformInfo.ini b/Engine/Config/Linux/DataDrivenPlatformInfo.ini
--- a/Engine/Config/Linux/DataDrivenPlatformInfo.ini
+++ b/Engine/Config/Linux/DataDrivenPlatformInfo.ini
@@ -16,3 +16,43 @@ bUATClosesAfterLaunch=true
 PlatformGroupName=Desktop
 PlatformSubMenu=Linux
 
+[ShaderPlatform VULKAN_SM5]
+bIsMobile = false
+bIsMetalMRT = false
+bIsPC = true
+bIsConsole = false
+bIsAndroidOpenGLES = false
+bSupportsMobileMultiView = false	; //#todo-vulkanddpi
+bSupportsArrayTextureCompression = false	; //#todo-vulkanddpi
+bSupportsVolumeTextureCompression = false	; //#todo-vulkanddpi
+bSupportsDistanceFields = true
+bSupportsDiaphragmDOF = true
+bSupportsRGBColorBuffer = false	; //#todo-vulkanddpi
+bSupportsCapsuleShadows = true
+bSupportsVolumetricFog = true
+bSupportsIndexBufferUAVs = true
+bSupportsInstancedStereo = false	; //#todo-vulkanddpi
+bSupportsMultiView = false	; //#todo-vulkanddpi
+bSupportsMSAA = true
+bSupports4ComponentUAVReadWrite = false	; //#todo-vulkanddpi
+bSupportsRenderTargetWriteMask = false	; //#todo-vulkanddpi
+bSupportsRayTracing = false
+bSupportsRayTracingIndirectInstanceData = false
+bSupportsGPUSkinCache = true
+bSupportsByteBufferComputeShaders = false	; //#todo-vulkanddpi
+bSupportsGPUScene = true
+bSupportsPrimitiveShaders = false	; //#todo-vulkanddpi
+bSupportsUInt64ImageAtomics = false	; //#todo-vulkanddpi
+bSupportsNanite = false
+bSupportsLumenGI = false
+bSupportsSSDIndirect = false	; //#todo-vulkanddpi
+bSupportsTemporalHistoryUpscale = false	; //#todo-vulkanddpi
+bSupportsRTIndexFromVS = false	; //#todo-vulkanddpi
+bSupportsIntrinsicWaveOnce = false	; //#todo-vulkanddpi
+bSupportsConservativeRasterization = false	; //#todo-vulkanddpi
+bSupportsWaveOperations = false	; //#todo-vulkanddpi
+bRequiresExplicit128bitRT = false	; //#todo-vulkanddpi
+bSupportsGen5TemporalAA = false	; //#todo-vulkanddpi
+bTargetsTiledGPU = false
+bNeedsOfflineCompiler = false
+bSupportsComputeFramework = false	; //#todo-vulkanddpi
diff --git a/Engine/Config/Windows/DataDrivenPlatformInfo.ini b/Engine/Config/Windows/DataDrivenPlatformInfo.ini
--- a/Engine/Config/Windows/DataDrivenPlatformInfo.ini
+++ b/Engine/Config/Windows/DataDrivenPlatformInfo.ini
@@ -28,4 +28,45 @@ bSupportsRayTracingIndirectInstanceData=true
 bSupportsWaveOperations=false
 bSupportsGPUScene=true
 bSupportsNanite=true
-bSupportsComputeFramework=true
\ No newline at end of file
+bSupportsComputeFramework=true
+
+[ShaderPlatform VULKAN_SM5]
+bIsMobile = false
+bIsMetalMRT = false
+bIsPC = true
+bIsConsole = false
+bIsAndroidOpenGLES = false
+bSupportsMobileMultiView = false	; //#todo-vulkanddpi
+bSupportsArrayTextureCompression = false	; //#todo-vulkanddpi
+bSupportsVolumeTextureCompression = false	; //#todo-vulkanddpi
+bSupportsDistanceFields = true
+bSupportsDiaphragmDOF = true
+bSupportsRGBColorBuffer = false	; //#todo-vulkanddpi
+bSupportsCapsuleShadows = true
+bSupportsVolumetricFog = true
+bSupportsIndexBufferUAVs = true
+bSupportsInstancedStereo = false	; //#todo-vulkanddpi
+bSupportsMultiView = false	; //#todo-vulkanddpi
+bSupportsMSAA = true
+bSupports4ComponentUAVReadWrite = false	; //#todo-vulkanddpi
+bSupportsRenderTargetWriteMask = false	; //#todo-vulkanddpi
+bSupportsRayTracing = false
+bSupportsRayTracingIndirectInstanceData = false
+bSupportsGPUSkinCache = true
+bSupportsByteBufferComputeShaders = false	; //#todo-vulkanddpi
+bSupportsGPUScene = true
+bSupportsPrimitiveShaders = false	; //#todo-vulkanddpi
+bSupportsUInt64ImageAtomics = false	; //#todo-vulkanddpi
+bSupportsNanite = false
+bSupportsLumenGI = false
+bSupportsSSDIndirect = false	; //#todo-vulkanddpi
+bSupportsTemporalHistoryUpscale = false	; //#todo-vulkanddpi
+bSupportsRTIndexFromVS = false	; //#todo-vulkanddpi
+bSupportsIntrinsicWaveOnce = false	; //#todo-vulkanddpi
+bSupportsConservativeRasterization = false	; //#todo-vulkanddpi
+bSupportsWaveOperations = false	; //#todo-vulkanddpi
+bRequiresExplicit128bitRT = false	; //#todo-vulkanddpi
+bSupportsGen5TemporalAA = false	; //#todo-vulkanddpi
+bTargetsTiledGPU = false
+bNeedsOfflineCompiler = false
+bSupportsComputeFramework = false	; //#todo-vulkanddpi
diff --git a/Engine/Source/Runtime/RHI/Private/RHI.cpp b/Engine/Source/Runtime/RHI/Private/RHI.cpp
--- a/Engine/Source/Runtime/RHI/Private/RHI.cpp
+++ b/Engine/Source/Runtime/RHI/Private/RHI.cpp
@@ -1592,7 +1592,6 @@ void FGenericDataDrivenShaderPlatformInfo::ParseDataDrivenShaderInfo(const FConf
 	Info.bSupportsRTIndexFromVS = GetSectionBool(Section, "bSupportsRTIndexFromVS");
 	Info.bSupportsIntrinsicWaveOnce = GetSectionBool(Section, "bSupportsIntrinsicWaveOnce");
 	Info.bSupportsConservativeRasterization = GetSectionBool(Section, "bSupportsConservativeRasterization");
-	Info.bSupportsGPUScene = GetSectionBool(Section, "bSupportsGPUScene");
 	Info.bSupportsWaveOperations = GetSectionBool(Section, "bSupportsWaveOperations");
 	Info.bRequiresExplicit128bitRT = GetSectionBool(Section, "bRequiresExplicit128bitRT");
 	Info.bSupportsGen5TemporalAA = GetSectionBool(Section, "bSupportsGen5TemporalAA");
