UE5 - Fixed ini for Vulkan SM5 (comments are at the beginning of the line, so it was failing to read 'true')

diff --git a/Engine/Config/Windows/DataDrivenPlatformInfo.ini b/Engine/Config/Windows/DataDrivenPlatformInfo.ini
--- a/Engine/Config/Windows/DataDrivenPlatformInfo.ini
+++ b/Engine/Config/Windows/DataDrivenPlatformInfo.ini
@@ -62,37 +62,111 @@ bIsMetalMRT = false
 bIsPC = true
 bIsConsole = false
 bIsAndroidOpenGLES = false
-bSupportsMobileMultiView = false	; //#todo-vulkanddpi
-bSupportsArrayTextureCompression = false	; //#todo-vulkanddpi
-bSupportsVolumeTextureCompression = false	; //#todo-vulkanddpi
+
+; //#todo-vulkanddpi
+bSupportsMobileMultiView = false
+
+; //#todo-vulkanddpi
+bSupportsArrayTextureCompression = false
+
+; //#todo-vulkanddpi
+bSupportsVolumeTextureCompression = false
+
 bSupportsDistanceFields = true
 bSupportsDiaphragmDOF = true
-bSupportsRGBColorBuffer = false	; //#todo-vulkanddpi
+
+; //#todo-vulkanddpi
+bSupportsRGBColorBuffer = false
+
 bSupportsCapsuleShadows = true
 bSupportsVolumetricFog = true
 bSupportsIndexBufferUAVs = true
-bSupportsInstancedStereo = false	; //#todo-vulkanddpi
-bSupportsMultiView = false	; //#todo-vulkanddpi
+
+; //#todo-vulkanddpi
+bSupportsInstancedStereo = false
+
+; //#todo-vulkanddpi
+bSupportsMultiView = false
+
 bSupportsMSAA = true
-bSupports4ComponentUAVReadWrite = false	; //#todo-vulkanddpi
-bSupportsRenderTargetWriteMask = false	; //#todo-vulkanddpi
+
+; //#todo-vulkanddpi
+bSupports4ComponentUAVReadWrite = false
+
+; //#todo-vulkanddpi
+bSupportsRenderTargetWriteMask = false
+
 bSupportsRayTracing = false
 bSupportsRayTracingIndirectInstanceData = false
 bSupportsGPUSkinCache = true
-bSupportsByteBufferComputeShaders = false	; //#todo-vulkanddpi
+
+; //#todo-vulkanddpi
+bSupportsByteBufferComputeShaders = false
+
 bSupportsGPUScene = true
-bSupportsPrimitiveShaders = false	; //#todo-vulkanddpi
-bSupportsUInt64ImageAtomics = false	; //#todo-vulkanddpi
+
+; //#todo-vulkanddpi
+bSupportsPrimitiveShaders = false
+
+; //#todo-vulkanddpi
+bSupportsUInt64ImageAtomics = false
+
+; //#todo-vulkanddpi
+bRequiresVendorExtensionsForAtomics = false
+
+; //#todo-vulkanddpi
 bSupportsNanite = false
+
+; //#todo-vulkanddpi
 bSupportsLumenGI = false
-bSupportsSSDIndirect = false	; //#todo-vulkanddpi
-bSupportsTemporalHistoryUpscale = false	; //#todo-vulkanddpi
-bSupportsRTIndexFromVS = false	; //#todo-vulkanddpi
-bSupportsIntrinsicWaveOnce = false	; //#todo-vulkanddpi
-bSupportsConservativeRasterization = false	; //#todo-vulkanddpi
-bSupportsWaveOperations = false	; //#todo-vulkanddpi
-bRequiresExplicit128bitRT = false	; //#todo-vulkanddpi
-bSupportsGen5TemporalAA = false	; //#todo-vulkanddpi
+
+; //#todo-vulkanddpi
+bSupportsSSDIndirect = false
+
+; //#todo-vulkanddpi
+bSupportsTemporalHistoryUpscale = false
+
+; //#todo-vulkanddpi
+bSupportsRTIndexFromVS = false
+
+; //#todo-vulkanddpi
+bSupportsIntrinsicWaveOnce = false
+
+; //#todo-vulkanddpi
+bSupportsConservativeRasterization = false
+
+
+; //#todo-vulkanddpi
+bSupportsWaveOperations = false
+
+; //#todo-vulkanddpi
+bRequiresExplicit128bitRT = false
+
+; //#todo-vulkanddpi
+bSupportsGen5TemporalAA = false
+
 bTargetsTiledGPU = false
 bNeedsOfflineCompiler = false
-bSupportsComputeFramework = false	; //#todo-vulkanddpi
+
+; //#todo-vulkanddpi
+bSupportsComputeFramework = false
+
+; //#todo-vulkanddpi
+bSupportsDualSourceBlending = false
+
+; //#todo-vulkanddpi
+bRequiresGeneratePrevTransformBuffer = false
+
+bRequiresRenderTargetDuringRaster = true
+
+; //#todo-vulkanddpi
+bRequiresDisableForwardLocalLights = true
+
+bCompileSignalProcessingPipeline = true
+bSupportsTessellation = true
+
+; //#todo-vulkanddpi
+bSupportsPerPixelDBufferMask = false
+
+bIsHlslcc = false
+NumberOfComputeThreads = 64
