UE5EA - Enable byte buffer shaders on Vulkan

diff --git a/Engine/Config/Linux/DataDrivenPlatformInfo.ini b/Engine/Config/Linux/DataDrivenPlatformInfo.ini
--- a/Engine/Config/Linux/DataDrivenPlatformInfo.ini
+++ b/Engine/Config/Linux/DataDrivenPlatformInfo.ini
@@ -53,9 +53,7 @@ bSupportsRenderTargetWriteMask = false
 bSupportsRayTracing = false
 bSupportsRayTracingIndirectInstanceData = false
 bSupportsGPUSkinCache = true
-
-; //#todo-vulkanddpi
-bSupportsByteBufferComputeShaders = false
+bSupportsByteBufferComputeShaders = true
 
 bSupportsGPUScene = true
 
diff --git a/Engine/Config/Windows/DataDrivenPlatformInfo.ini b/Engine/Config/Windows/DataDrivenPlatformInfo.ini
--- a/Engine/Config/Windows/DataDrivenPlatformInfo.ini
+++ b/Engine/Config/Windows/DataDrivenPlatformInfo.ini
@@ -92,9 +92,7 @@ bSupportsRenderTargetWriteMask = false
 bSupportsRayTracing = false
 bSupportsRayTracingIndirectInstanceData = false
 bSupportsGPUSkinCache = true
-
-; //#todo-vulkanddpi
-bSupportsByteBufferComputeShaders = false
+bSupportsByteBufferComputeShaders = true
 
 bSupportsGPUScene = true
 
