Controls whether Gen5TAA is supported solely based on FDataDrivenShaderPlatformInfo

diff --git a/Engine/Config/Linux/DataDrivenPlatformInfo.ini b/Engine/Config/Linux/DataDrivenPlatformInfo.ini
--- a/Engine/Config/Linux/DataDrivenPlatformInfo.ini
+++ b/Engine/Config/Linux/DataDrivenPlatformInfo.ini
@@ -115,6 +115,7 @@ bRequiresDisableForwardLocalLights = true
 
 bCompileSignalProcessingPipeline = true
 bSupportsTessellation = true
+bSupportsGen5TemporalAA=true
 
 ; //#todo-vulkanddpi
 bSupportsPerPixelDBufferMask = false
diff --git a/Engine/Config/Windows/DataDrivenPlatformInfo.ini b/Engine/Config/Windows/DataDrivenPlatformInfo.ini
--- a/Engine/Config/Windows/DataDrivenPlatformInfo.ini
+++ b/Engine/Config/Windows/DataDrivenPlatformInfo.ini
@@ -39,6 +39,7 @@ bSupportsSSDIndirect=true
 bRequiresVendorExtensionsForAtomics=true
 bSupportsAnisotropicMaterials=true
 bSupportsGPUSkinCache=true
+bSupportsGen5TemporalAA=true
 
 [PreviewPlatform PCD3D_SM5]
 PlatformName=PCD3D_SM5
@@ -66,6 +67,9 @@ bSupportsGPUSkinCache=true
 bSupportsMeshShaders=true
 FriendlyName=LOCTEXT("FriendlyShaderPlatformName_PCD3D_SM6", "Shader Model 6")
 
+; TODO(UE-111389)
+bSupportsGen5TemporalAA=false
+
 [PreviewPlatform PCD3D_SM6]
 PlatformName=PCD3D_SM6
 ShaderFormat=PCD3D_SM6
@@ -172,6 +176,7 @@ bRequiresDisableForwardLocalLights = true
 
 bCompileSignalProcessingPipeline = true
 bSupportsTessellation = true
+bSupportsGen5TemporalAA=true
 
 ; //#todo-vulkanddpi
 bSupportsPerPixelDBufferMask = false
diff --git a/Engine/Source/Runtime/Renderer/Private/PostProcess/TemporalAA.cpp b/Engine/Source/Runtime/Renderer/Private/PostProcess/TemporalAA.cpp
--- a/Engine/Source/Runtime/Renderer/Private/PostProcess/TemporalAA.cpp
+++ b/Engine/Source/Runtime/Renderer/Private/PostProcess/TemporalAA.cpp
@@ -104,10 +104,7 @@ inline bool DoesPlatformSupportTemporalHistoryUpscale(EShaderPlatform Platform)
 
 inline bool DoesPlatformSupportGen5TAA(EShaderPlatform Platform)
 {
-	// TODO(UE-111389): Fixes artifact on SP_PCD3D_SM6
-	return (
-		(IsPCPlatform(Platform) && IsFeatureLevelSupported(Platform, ERHIFeatureLevel::SM5) && Platform != SP_PCD3D_SM6) ||
-		FDataDrivenShaderPlatformInfo::GetSupportsGen5TemporalAA(Platform));
+	return FDataDrivenShaderPlatformInfo::GetSupportsGen5TemporalAA(Platform);
 }
 
 BEGIN_SHADER_PARAMETER_STRUCT(FTAACommonParameters, )
