Add missing data driven definition of hair strands support for PC SM5.

diff --git a/Engine/Config/Windows/DataDrivenPlatformInfo.ini b/Engine/Config/Windows/DataDrivenPlatformInfo.ini
--- a/Engine/Config/Windows/DataDrivenPlatformInfo.ini
+++ b/Engine/Config/Windows/DataDrivenPlatformInfo.ini
@@ -51,6 +51,7 @@ bSupportsAnisotropicMaterials=true
 bSupportsVariableRateShading=true
 bSupportsGPUSkinCache=true
 bSupportsGen5TemporalAA=true
+bSupportsHairStrandGeometry=true
 
 [PreviewPlatform PCD3D_SM5]
 PlatformName=PCD3D_SM5
