Set bSupportsVariableRateShading to false in DataDrivePlatformInfo.ini for VULKAN_SM5 temporarily.  Our current setup has VULKAN_SM5 defined for multiple other platforms and they clobber each other on load (you keep the settings of the last one loaded it seems).

diff --git a/Engine/Config/Windows/DataDrivenPlatformInfo.ini b/Engine/Config/Windows/DataDrivenPlatformInfo.ini
--- a/Engine/Config/Windows/DataDrivenPlatformInfo.ini
+++ b/Engine/Config/Windows/DataDrivenPlatformInfo.ini
@@ -193,7 +193,7 @@ bSupportsPerPixelDBufferMask = false
 
 bIsHlslcc = false
 
-bSupportsVariableRateShading = true
+bSupportsVariableRateShading = false
 
 NumberOfComputeThreads = 64
 
