Unify the DataDrivePlatformInfo for VULKAN_SM5 in a common VulkanPC fake platform.  The same slot was used in Editor (FGenericDataDrivenShaderPlatformInfo::Infos[SP_VULKAN_SM5]) for all these different sets of values (Windows, Linux, etc), so they would overwrite eachother (and the last one loaded was used).  Use a single set of values for all VULKAN_SM5 and we will add a ShaderPlatform for those that need to diverge (like it was done for Android) or store the setting elsewhere.

diff --git a/Engine/Config/Linux/DataDrivenPlatformInfo.ini b/Engine/Config/Linux/DataDrivenPlatformInfo.ini
--- a/Engine/Config/Linux/DataDrivenPlatformInfo.ini
+++ b/Engine/Config/Linux/DataDrivenPlatformInfo.ini
@@ -24,102 +24,3 @@ bUATClosesAfterLaunch=true
 PlatformGroupName=Desktop
 PlatformSubMenu=Linux
 AudioCompressionSettingsIniSectionName=/Script/LinuxTargetPlatform.LinuxTargetSettings
-
-[ShaderPlatform VULKAN_SM5]
-bIsMobile = false
-bIsMetalMRT = false
-bIsPC = true
-bIsConsole = false
-bIsAndroidOpenGLES = false
-
-; //#todo-vulkanddpi
-bSupportsMobileMultiView = false
-
-; //#todo-vulkanddpi
-bSupportsArrayTextureCompression = false
-
-bSupportsDistanceFields = true
-bSupportsDiaphragmDOF = true
-bSupportsRGBColorBuffer = true
-bSupportsCapsuleShadows = true
-bSupportsVolumetricFog = true
-bSupportsIndexBufferUAVs = true
-
-; //#todo-vulkanddpi
-bSupportsInstancedStereo = false
-
-; //#todo-vulkanddpi
-bSupportsMultiView = false
-
-bSupportsMSAA = true
-
-; //#todo-vulkanddpi
-bSupports4ComponentUAVReadWrite = false
-
-; //#todo-vulkanddpi
-bSupportsRenderTargetWriteMask = false
-
-bSupportsRayTracing = false
-bSupportsRayTracingIndirectInstanceData = false
-bSupportsByteBufferComputeShaders = true
-
-bSupportsGPUScene = true
-
-; //#todo-vulkanddpi
-bSupportsPrimitiveShaders = false
-
-bSupportsUInt64ImageAtomics = true
-
-; //#todo-vulkanddpi
-bRequiresVendorExtensionsForAtomics = false
-
-bSupportsNanite = true
-bSupportsLumenGI = true
-
-; //#todo-vulkanddpi
-bSupportsSSDIndirect = false
-
-; //#todo-vulkanddpi
-bSupportsTemporalHistoryUpscale = false
-
-; //#todo-vulkanddpi
-bSupportsRTIndexFromVS = false
-
-; //#todo-vulkanddpi
-bSupportsIntrinsicWaveOnce = false
-
-; //#todo-vulkanddpi
-bSupportsConservativeRasterization = false
-
-; //#todo-vulkanddpi
-bSupportsWaveOperations = false
-
-; //#todo-vulkanddpi
-bRequiresExplicit128bitRT = false
-
-bTargetsTiledGPU = false
-bNeedsOfflineCompiler = false
-
-; //#todo-vulkanddpi
-bSupportsComputeFramework = false
-
-; //#todo-vulkanddpi
-bSupportsDualSourceBlending = false
-
-; //#todo-vulkanddpi
-bRequiresGeneratePrevTransformBuffer = false
-
-bRequiresRenderTargetDuringRaster = true
-
-; //#todo-vulkanddpi
-bRequiresDisableForwardLocalLights = true
-
-bCompileSignalProcessingPipeline = true
-bSupportsGen5TemporalAA=true
-bSupportsFFTBloom=true
-
-; //#todo-vulkanddpi
-bSupportsPerPixelDBufferMask = false
-
-bIsHlslcc = false
-NumberOfComputeThreads = 64
diff --git a/Engine/Config/Unix/DataDrivenPlatformInfo.ini b/Engine/Config/Unix/DataDrivenPlatformInfo.ini
new file mode 100644
--- /dev/null
+++ b/Engine/Config/Unix/DataDrivenPlatformInfo.ini
@@ -0,0 +1,5 @@
+[DataDrivenPlatformInfo]
+bIsConfidential=false
+IniParent=VulkanPC
+bIsFakePlatform=true
+GlobalIdentifier=969267E8F7A6414DAD0450C1CAE83A6E
\ No newline at end of file
diff --git a/Engine/Config/VulkanPC/DataDrivenPlatformInfo.ini b/Engine/Config/VulkanPC/DataDrivenPlatformInfo.ini
new file mode 100644
--- /dev/null
+++ b/Engine/Config/VulkanPC/DataDrivenPlatformInfo.ini
@@ -0,0 +1,115 @@
+; Have a single source of truth for all desktop class Vulkan platforms
+
+; This is done so all the platforms using SP_VULKAN_SM5 can share the same DDC
+; If your Vulkan platform needs its own settings, create a new SP for it
+
+[DataDrivenPlatformInfo]
+bIsConfidential=false
+bIsFakePlatform=true
+GlobalIdentifier=97F8C85EF1B0424792D201C893358C8F
+
+[ShaderPlatform VULKAN_SM5]
+bIsMobile = false
+bIsMetalMRT = false
+bIsPC = true
+bIsConsole = false
+bIsAndroidOpenGLES = false
+
+; //#todo-vulkanddpi
+bSupportsMobileMultiView = false
+
+; //#todo-vulkanddpi
+bSupportsArrayTextureCompression = false
+
+bSupportsDistanceFields = true
+bSupportsDiaphragmDOF = true
+bSupportsRGBColorBuffer = true
+bSupportsCapsuleShadows = true
+bSupportsVolumetricFog = true
+bSupportsIndexBufferUAVs = true
+
+; //#todo-vulkanddpi
+bSupportsInstancedStereo = false
+
+; //#todo-vulkanddpi
+bSupportsMultiView = false
+
+bSupportsMSAA = true
+
+; //#todo-vulkanddpi
+bSupports4ComponentUAVReadWrite = false
+
+; //#todo-vulkanddpi
+bSupportsRenderTargetWriteMask = false
+
+bSupportsRayTracing = false
+bSupportsRayTracingIndirectInstanceData = false
+bSupportsGPUSkinCache = true
+bSupportsByteBufferComputeShaders = true
+
+bSupportsGPUScene = true
+
+; //#todo-vulkanddpi
+bSupportsPrimitiveShaders = false
+
+bSupportsUInt64ImageAtomics = true
+
+; //#todo-vulkanddpi
+bRequiresVendorExtensionsForAtomics = false
+
+bSupportsNanite = true
+bSupportsLumenGI = true
+
+; //#todo-vulkanddpi
+bSupportsSSDIndirect = false
+
+; //#todo-vulkanddpi
+bSupportsTemporalHistoryUpscale = false
+
+; //#todo-vulkanddpi
+bSupportsRTIndexFromVS = false
+
+; //#todo-vulkanddpi
+bSupportsIntrinsicWaveOnce = false
+
+; //#todo-vulkanddpi
+bSupportsConservativeRasterization = false
+
+
+; //#todo-vulkanddpi
+bSupportsWaveOperations = false
+
+; //#todo-vulkanddpi
+bRequiresExplicit128bitRT = false
+
+bTargetsTiledGPU = false
+bNeedsOfflineCompiler = false
+
+; //#todo-vulkanddpi
+bSupportsComputeFramework = false
+
+; //#todo-vulkanddpi
+bSupportsDualSourceBlending = false
+
+; //#todo-vulkanddpi
+bRequiresGeneratePrevTransformBuffer = false
+
+bRequiresRenderTargetDuringRaster = true
+
+; //#todo-vulkanddpi
+bRequiresDisableForwardLocalLights = true
+
+bCompileSignalProcessingPipeline = true
+bSupportsGen5TemporalAA=true
+
+; //#todo-vulkanddpi
+bSupportsPerPixelDBufferMask = false
+
+bIsHlslcc = false
+
+bSupportsVariableRateShading = false
+
+NumberOfComputeThreads = 64
+
+bSupportsHairStrandGeometry=true
+bSupportsFFTBloom=true
diff --git a/Engine/Config/Windows/DataDrivenPlatformInfo.ini b/Engine/Config/Windows/DataDrivenPlatformInfo.ini
--- a/Engine/Config/Windows/DataDrivenPlatformInfo.ini
+++ b/Engine/Config/Windows/DataDrivenPlatformInfo.ini
@@ -1,6 +1,6 @@
 [DataDrivenPlatformInfo]
 bIsConfidential=false
-PlatformNameAliases=Win64
+IniParent=VulkanPC
 TargetSettingsIniSectionName=/Script/WindowsTargetPlatform.WindowsTargetSettings
 Freezing_b32Bit=false
 Freezing_bWithRayTracing=true
@@ -94,112 +94,6 @@ ShaderFormat=PCD3D_SM6
 MenuText=LOCTEXT("PreviewMenuText_PCD3D_SM6", "Shader Model 6")
 MenuTooltip=LOCTEXT("PreviewMenuTooltip_PCD3D_SM6", "Preview using SM6 device profile.")
 
-
-[ShaderPlatform VULKAN_SM5]
-bIsMobile = false
-bIsMetalMRT = false
-bIsPC = true
-bIsConsole = false
-bIsAndroidOpenGLES = false
-
-; //#todo-vulkanddpi
-bSupportsMobileMultiView = false
-
-; //#todo-vulkanddpi
-bSupportsArrayTextureCompression = false
-
-bSupportsDistanceFields = true
-bSupportsDiaphragmDOF = true
-bSupportsRGBColorBuffer = true
-bSupportsCapsuleShadows = true
-bSupportsVolumetricFog = true
-bSupportsIndexBufferUAVs = true
-
-; //#todo-vulkanddpi
-bSupportsInstancedStereo = false
-
-; //#todo-vulkanddpi
-bSupportsMultiView = false
-
-bSupportsMSAA = true
-
-; //#todo-vulkanddpi
-bSupports4ComponentUAVReadWrite = false
-
-; //#todo-vulkanddpi
-bSupportsRenderTargetWriteMask = false
-
-bSupportsRayTracing = false
-bSupportsRayTracingIndirectInstanceData = false
-bSupportsByteBufferComputeShaders = true
-
-bSupportsGPUScene = true
-
-; //#todo-vulkanddpi
-bSupportsPrimitiveShaders = false
-
-bSupportsUInt64ImageAtomics = true
-
-; //#todo-vulkanddpi
-bRequiresVendorExtensionsForAtomics = false
-
-bSupportsNanite = true
-bSupportsLumenGI = true
-
-; //#todo-vulkanddpi
-bSupportsSSDIndirect = false
-
-; //#todo-vulkanddpi
-bSupportsTemporalHistoryUpscale = false
-
-; //#todo-vulkanddpi
-bSupportsRTIndexFromVS = false
-
-; //#todo-vulkanddpi
-bSupportsIntrinsicWaveOnce = false
-
-; //#todo-vulkanddpi
-bSupportsConservativeRasterization = false
-
-
-; //#todo-vulkanddpi
-bSupportsWaveOperations = false
-
-; //#todo-vulkanddpi
-bRequiresExplicit128bitRT = false
-
-bTargetsTiledGPU = false
-bNeedsOfflineCompiler = false
-
-; //#todo-vulkanddpi
-bSupportsComputeFramework = false
-
-; //#todo-vulkanddpi
-bSupportsDualSourceBlending = false
-
-; //#todo-vulkanddpi
-bRequiresGeneratePrevTransformBuffer = false
-
-bRequiresRenderTargetDuringRaster = true
-
-; //#todo-vulkanddpi
-bRequiresDisableForwardLocalLights = true
-
-bCompileSignalProcessingPipeline = true
-bSupportsGen5TemporalAA=true
-
-; //#todo-vulkanddpi
-bSupportsPerPixelDBufferMask = false
-
-bIsHlslcc = false
-
-bSupportsVariableRateShading = false
-
-NumberOfComputeThreads = 64
-
-bSupportsHairStrandGeometry=true
-bSupportsFFTBloom=true
-
 [ShaderPlatform PCD3D_ES3_1]
 Language=D3D
 MaxFeatureLevel=ES3_1
