Fix bad merge.

diff --git a/Engine/Config/Windows/DataDrivenPlatformInfo.ini b/Engine/Config/Windows/DataDrivenPlatformInfo.ini
--- a/Engine/Config/Windows/DataDrivenPlatformInfo.ini
+++ b/Engine/Config/Windows/DataDrivenPlatformInfo.ini
@@ -1,6 +1,7 @@
 [DataDrivenPlatformInfo]
 bIsConfidential=false
 IniParent=VulkanPC
+PlatformNameAliases=Win64
 TargetSettingsIniSectionName=/Script/WindowsTargetPlatform.WindowsTargetSettings
 Freezing_b32Bit=false
 Freezing_bWithRayTracing=true
