Fix UBT configuration reading code to read the editor per project user setting from the corrected saved file location.
This resolves a UBT issue where building from Visual Studio/Rider while the editor is running wouldn't automatically change to HotReload mode.

diff --git a/Engine/Config/BaseEditorPerProjectUserSettings.ini b/Engine/Config/BaseEditorPerProjectUserSettings.ini
--- a/Engine/Config/BaseEditorPerProjectUserSettings.ini
+++ b/Engine/Config/BaseEditorPerProjectUserSettings.ini
@@ -926,7 +926,7 @@ MaxBitrate=4000
 
 [/Script/LiveCoding.LiveCodingSettings]
 Startup=AutomaticButHidden
-bEnabled=True
+bEnabled=False
 bEnableReinstancing=True
 bAutomaticallyCompileNewClasses=True
 
diff --git a/Engine/Config/Windows/WindowsEditorPerProjectUserSettings.ini b/Engine/Config/Windows/WindowsEditorPerProjectUserSettings.ini
new file mode 100644
--- /dev/null
+++ b/Engine/Config/Windows/WindowsEditorPerProjectUserSettings.ini
@@ -0,0 +1,22 @@
+
+; -------------------------------------------------------------------------------------------------
+; IMPORTANT:  EditorPerProjectUserSettings.ini has special behavior!
+;
+; This .ini file contains the defaults for many editor and engine "preferences".  These preferences
+; are saved to an EditorPerProjectUserSettings.ini file in the user's /<Game>/Saved/Config/ directory. 
+;
+; If you change the default for a preference in here, that change will *NOT* automatically 
+; propagate to users who already have saved preferences.  The idea is that we must preserve
+; their existing settings and never want to clobber entries in that file.
+;
+; If you add a new preference, the new key and value *WILL* be propagated to the user's
+; EditorPerProjectUserSettings.ini file.  After we load the user's existing settings, we'll merge in
+; any missing keys and values that are present in these defaults.
+;
+; One easy technique for forcing a "reset" of an already-existing preference is to simply
+; rename the variable for that preference.  The config system will treat it as a newly-added
+; key and will propagate the change to existing users' preferences.
+; -------------------------------------------------------------------------------------------------
+
+[/Script/LiveCoding.LiveCodingSettings]
+bEnabled=True
diff --git a/Engine/Source/Programs/UnrealBuildTool/System/ConfigHierarchy.cs b/Engine/Source/Programs/UnrealBuildTool/System/ConfigHierarchy.cs
--- a/Engine/Source/Programs/UnrealBuildTool/System/ConfigHierarchy.cs
+++ b/Engine/Source/Programs/UnrealBuildTool/System/ConfigHierarchy.cs
@@ -1144,6 +1144,10 @@ namespace UnrealBuildTool
 			{
 				yield return FileReference.Combine(GetGameAgnosticSavedDir(), "Config", PlatformName + "Editor", BaseIniName + ".ini");
 			}
+			else if (Type == ConfigHierarchyType.EditorPerProjectUserSettings)
+			{
+				yield return FileReference.Combine(GetGeneratedConfigDir(ProjectDir), PlatformName + "Editor", BaseIniName + ".ini");
+			}
 			else
 			{
 				yield return FileReference.Combine(GetGeneratedConfigDir(ProjectDir), PlatformName, BaseIniName + ".ini");
diff --git a/Engine/Source/Programs/UnrealBuildTool/System/HotReload.cs b/Engine/Source/Programs/UnrealBuildTool/System/HotReload.cs
--- a/Engine/Source/Programs/UnrealBuildTool/System/HotReload.cs
+++ b/Engine/Source/Programs/UnrealBuildTool/System/HotReload.cs
@@ -277,11 +277,14 @@ namespace UnrealBuildTool
 			}
 
 			// Check if we're using LiveCode instead
-			ConfigHierarchy EditorPerProjectHierarchy = ConfigCache.ReadHierarchy(ConfigHierarchyType.EditorPerProjectUserSettings, DirectoryReference.FromFile(TargetDesc.ProjectFile), TargetDesc.Platform);
-			bool bEnableLiveCode;
-			if(EditorPerProjectHierarchy.GetBool("/Script/LiveCoding.LiveCodingSettings", "bEnabled", out bEnableLiveCode) && bEnableLiveCode)
+			if (BuildHostPlatform.Current.Platform == UnrealTargetPlatform.Win64) // Temporary - new 5.0 projects will have live coding setting on for platforms that don't support it.
 			{
-				return false;
+				ConfigHierarchy EditorPerProjectHierarchy = ConfigCache.ReadHierarchy(ConfigHierarchyType.EditorPerProjectUserSettings, DirectoryReference.FromFile(TargetDesc.ProjectFile), TargetDesc.Platform);
+				bool bEnableLiveCode;
+				if(EditorPerProjectHierarchy.GetBool("/Script/LiveCoding.LiveCodingSettings", "bEnabled", out bEnableLiveCode) && bEnableLiveCode)
+				{
+					return false;
+				}
 			}
 
 			bool bIsRunning = false;
