Temporaly disable TSR on desktop SPIRV backends due to long compilation time in DXC

diff --git a/Engine/Config/Mac/DataDrivenPlatformInfo.ini b/Engine/Config/Mac/DataDrivenPlatformInfo.ini
--- a/Engine/Config/Mac/DataDrivenPlatformInfo.ini
+++ b/Engine/Config/Mac/DataDrivenPlatformInfo.ini
@@ -31,7 +31,7 @@ bSupportsDebugViewShaders=true
 bSupportsAnisotropicMaterials=true
 bSupportsByteBufferComputeShaders=true
 bSupportsDistanceFields=true
-bSupportsGen5TemporalAA=true
+bSupportsGen5TemporalAA=false
 bSupportsLumenGI=true
 bSupportsSSDIndirect=true
 bSupportsWaveOperations=false
diff --git a/Engine/Config/VulkanPC/DataDrivenPlatformInfo.ini b/Engine/Config/VulkanPC/DataDrivenPlatformInfo.ini
--- a/Engine/Config/VulkanPC/DataDrivenPlatformInfo.ini
+++ b/Engine/Config/VulkanPC/DataDrivenPlatformInfo.ini
@@ -106,7 +106,7 @@ bRequiresRenderTargetDuringRaster = true
 bRequiresDisableForwardLocalLights = true
 
 bCompileSignalProcessingPipeline = true
-bSupportsGen5TemporalAA=true
+bSupportsGen5TemporalAA=false
 
 ; //#todo-vulkanddpi
 bSupportsPerPixelDBufferMask = false
diff --git a/Engine/Config/Windows/DataDrivenPlatformInfo.ini b/Engine/Config/Windows/DataDrivenPlatformInfo.ini
--- a/Engine/Config/Windows/DataDrivenPlatformInfo.ini
+++ b/Engine/Config/Windows/DataDrivenPlatformInfo.ini
@@ -78,7 +78,7 @@ bSupportsSSDIndirect=true
 bRequiresVendorExtensionsForAtomics=true
 bSupportsAnisotropicMaterials=true
 bSupportsVariableRateShading=true
-bSupportsGen5TemporalAA=true
+bSupportsGen5TemporalAA=false
 bSupportsHairStrandGeometry=true
 bSupportsFFTBloom=true
 bSupportsDiaphragmDOF=true
