Use the Compressed transforms on PC with SM5. Currently investigating more what is causing the issue in the Editor.

diff --git a/Engine/Config/Windows/DataDrivenPlatformInfo.ini b/Engine/Config/Windows/DataDrivenPlatformInfo.ini
--- a/Engine/Config/Windows/DataDrivenPlatformInfo.ini
+++ b/Engine/Config/Windows/DataDrivenPlatformInfo.ini
@@ -71,7 +71,7 @@ bSupportsHighEndRayTracingReflections=true
 bSupportsWaveOperations=RuntimeDependent
 bSupportsGPUScene=true
 bSupportsNanite=false
-bSupportsSceneDataCompressedTransforms=false
+bSupportsSceneDataCompressedTransforms=true
 bSupportsComputeFramework=true
 bSupportsByteBufferComputeShaders=true
 bSupportsLumenGI=true
