- Removed unused section in Windows DDPI.ini, leftover from UE4 cooked editor stuff

diff --git a/Engine/Config/Windows/DataDrivenPlatformInfo.ini b/Engine/Config/Windows/DataDrivenPlatformInfo.ini
--- a/Engine/Config/Windows/DataDrivenPlatformInfo.ini
+++ b/Engine/Config/Windows/DataDrivenPlatformInfo.ini
@@ -31,27 +31,6 @@ bUATClosesAfterLaunch=false
 PlatformGroupName=Desktop
 UBTPlatformName=Win64
 
-[PlatformInfo WindowsCookedEditor]
-TargetPlatformName=WindowsCookedEditor
-DisplayName=Windows Cooked Editor
-PlatformType=Editor
-PlatformFlags=None
-NormalIconPath=Launcher/Windows/Platform_Windows_24x
-LargeIconPath=Launcher/Windows/Platform_Windows_128x
-XLargeIconPath=
-UATCommandLine=
-AutoSDKPath=
-TutorialPath=SharingAndReleasing
-Windows:bIsEnabled=false
-Mac:bIsEnabled=false
-Linux:bIsEnabled=false
-BinariesDirectoryName=Win64
-IniPlatformName=Windows
-bUsesHostCompiler=true
-bUATClosesAfterLaunch=false
-UBTTargetID=Win64
-PlatformGroupName=Desktop
-
 
 [ShaderPlatform PCD3D_SM5]
 Language=D3D
