Fixed: broken geometry in editor mobile preview when "GPU Scene on Mobile" is true

diff --git a/Engine/Config/VulkanPC/DataDrivenPlatformInfo.ini b/Engine/Config/VulkanPC/DataDrivenPlatformInfo.ini
--- a/Engine/Config/VulkanPC/DataDrivenPlatformInfo.ini
+++ b/Engine/Config/VulkanPC/DataDrivenPlatformInfo.ini
@@ -140,3 +140,4 @@ bIsMobile=true
 bSupportsMobileMultiView=true
 bSupportsIndexBufferUAVs=true
 bSupportsManualVertexFetch=false
+bSupportsSceneDataCompressedTransforms=true
diff --git a/Engine/Config/Windows/DataDrivenPlatformInfo.ini b/Engine/Config/Windows/DataDrivenPlatformInfo.ini
--- a/Engine/Config/Windows/DataDrivenPlatformInfo.ini
+++ b/Engine/Config/Windows/DataDrivenPlatformInfo.ini
@@ -144,6 +144,7 @@ bSupportsDistanceFields = true
 bSupportsByteBufferComputeShaders = true
 bSupportsInstancedStereo=true
 bSupportsManualVertexFetch=false
+bSupportsSceneDataCompressedTransforms=true
 
 [ShaderPlatform OPENGL_PCES3_1]
 Language=OpenGL
