Use bindless samplers for DXR on Windows by default

* The non-bindless code path is planned to be deprecated in the future engine release

diff --git a/Engine/Config/Windows/WindowsEngine.ini b/Engine/Config/Windows/WindowsEngine.ini
--- a/Engine/Config/Windows/WindowsEngine.ini
+++ b/Engine/Config/Windows/WindowsEngine.ini
@@ -22,6 +22,8 @@ QualityLevelMapping="high"
 ; Enable single RHI thread stall for D3D11
 r.Occlusion.SingleRHIThreadStall=1
 r.Shadow.DetectVertexShaderLayerAtRuntime=1
+; Enable SM 6.6 bindless samplers for DXR on Windows by default (non-bindless path is planned for deprecation in the future)
+rhi.Bindless.Samplers=2
 
 [PlatformCrypto]
 PlatformRequiresDataCrypto=True
