Undo 23894521 - revert back to non-bindless samplers for DXR

* The CVar can't be used to enable bindless resources per platform

diff --git a/Engine/Config/Windows/WindowsEngine.ini b/Engine/Config/Windows/WindowsEngine.ini
--- a/Engine/Config/Windows/WindowsEngine.ini
+++ b/Engine/Config/Windows/WindowsEngine.ini
@@ -22,8 +22,6 @@ QualityLevelMapping="high"
 ; Enable single RHI thread stall for D3D11
 r.Occlusion.SingleRHIThreadStall=1
 r.Shadow.DetectVertexShaderLayerAtRuntime=1
-; Enable SM 6.6 bindless samplers for DXR on Windows by default (non-bindless path is planned for deprecation in the future)
-rhi.Bindless.Samplers=2
 
 [PlatformCrypto]
 PlatformRequiresDataCrypto=True
