User friendly labels for Shader Preview Levels

diff --git a/Engine/Config/Android/DataDrivenPlatformInfo.ini b/Engine/Config/Android/DataDrivenPlatformInfo.ini
--- a/Engine/Config/Android/DataDrivenPlatformInfo.ini
+++ b/Engine/Config/Android/DataDrivenPlatformInfo.ini
@@ -64,7 +64,7 @@ bSupportsMobileMultiView=true
 bNeedsOfflineCompiler=true
 bSupportsManualVertexFetch=false
 bSupportsClipDistance = false
-FriendlyName=LOCTEXT("FriendlyShaderPlatformName_AndroidES31", "Android ES 3_1")
+FriendlyName=LOCTEXT("FriendlyShaderPlatformName_AndroidES31", "Android Opengl ES31")
 
 [ShaderPlatform VULKAN_ES3_1_ANDROID]
 Language=Vulkan
@@ -77,7 +77,7 @@ bNeedsOfflineCompiler=true
 bSupportsIndexBufferUAVs=true
 bSupportsManualVertexFetch=false
 bSupportsClipDistance = false
-FriendlyName=LOCTEXT("FriendlyShaderPlatformName_AndroidVulkan", "Android Vulkan")
+FriendlyName=LOCTEXT("FriendlyShaderPlatformName_AndroidVulkan", "Android Vulkan ES31")
 
 [ShaderPlatform VULKAN_SM5_ANDROID]
 Language=Vulkan
diff --git a/Engine/Config/HoloLens/DataDrivenPlatformInfo.ini b/Engine/Config/HoloLens/DataDrivenPlatformInfo.ini
--- a/Engine/Config/HoloLens/DataDrivenPlatformInfo.ini
+++ b/Engine/Config/HoloLens/DataDrivenPlatformInfo.ini
@@ -34,4 +34,4 @@ bSupportsInstancedStereo=true
 bSupportsManualVertexFetch=false
 bSupportsVertexShaderLayer=true
 bSupportsClipDistance = false
-FriendlyName=LOCTEXT("FriendlyShaderPlatformName_HololensES31", "Hololens Shader Model 3_1")
\ No newline at end of file
+FriendlyName=LOCTEXT("FriendlyShaderPlatformName_HololensES31", "Hololens ES31")
\ No newline at end of file
diff --git a/Engine/Config/IOS/DataDrivenPlatformInfo.ini b/Engine/Config/IOS/DataDrivenPlatformInfo.ini
--- a/Engine/Config/IOS/DataDrivenPlatformInfo.ini
+++ b/Engine/Config/IOS/DataDrivenPlatformInfo.ini
@@ -54,7 +54,7 @@ bTargetsTiledGPU=true
 bSupports4ComponentUAVReadWrite=true
 bSupportsManualVertexFetch=false
 bSupportsVolumeTextureAtomics=false
-FriendlyName=LOCTEXT("FriendlyShaderPlatformName_IOSES31", "IOS Shader Model 3_1")
+FriendlyName=LOCTEXT("FriendlyShaderPlatformName_IOSES31", "IOS Metal ES31")
 
 [ShaderPlatform METAL_MRT]
 Language=Metal
@@ -66,4 +66,4 @@ bSupportsIndexBufferUAVs=true
 bSupports4ComponentUAVReadWrite=true
 bSupportsManualVertexFetch=false
 bSupportsVolumeTextureAtomics=false
-FriendlyName=LOCTEXT("FriendlyShaderPlatformName_IOSSM5", "IOS Shader Model 5")
\ No newline at end of file
+FriendlyName=LOCTEXT("FriendlyShaderPlatformName_IOSSM5", "IOS Metal SM5")
\ No newline at end of file
diff --git a/Engine/Config/Mac/DataDrivenPlatformInfo.ini b/Engine/Config/Mac/DataDrivenPlatformInfo.ini
--- a/Engine/Config/Mac/DataDrivenPlatformInfo.ini
+++ b/Engine/Config/Mac/DataDrivenPlatformInfo.ini
@@ -59,7 +59,7 @@ bSupportsInlineRayTracing = false
 bSupportsRayTracingIndirectInstanceData = false
 bSupportsByteBufferComputeShaders = true
 bSupportsHairStrandGeometry = false
-FriendlyName=LOCTEXT("FriendlyShaderPlatformName_MacSM5", "Mac Shader Model 5")
+FriendlyName=LOCTEXT("FriendlyShaderPlatformName_MacSM5", "Mac Metal SM5")
 
 [ShaderPlatform METAL_MRT_MAC]
 Language=Metal
@@ -72,7 +72,7 @@ bSupports4ComponentUAVReadWrite=true
 bSupportsManualVertexFetch=true
 bSupportsVolumeTextureAtomics=false
 bSupportsDualSourceBlending=true
-FriendlyName=LOCTEXT("FriendlyShaderPlatformName_MacSM5_MRT", "Mac Shader Model 5 MRT")
+FriendlyName=LOCTEXT("FriendlyShaderPlatformName_MacSM5_MRT", "Mac Metal SM5 MRT")
 
 [ShaderPlatform METAL_MACES3_1]
 Language=Metal
@@ -83,4 +83,4 @@ bIsMobile=true
 bSupports4ComponentUAVReadWrite=true
 bSupportsManualVertexFetch=false
 bSupportsVolumeTextureAtomics=false
-FriendlyName=LOCTEXT("FriendlyShaderPlatformName_MacES31", "Mac Shader Model 3.1")
+FriendlyName=LOCTEXT("FriendlyShaderPlatformName_MacES31", "Mac Metal ES31")
diff --git a/Engine/Config/TVOS/DataDrivenPlatformInfo.ini b/Engine/Config/TVOS/DataDrivenPlatformInfo.ini
--- a/Engine/Config/TVOS/DataDrivenPlatformInfo.ini
+++ b/Engine/Config/TVOS/DataDrivenPlatformInfo.ini
@@ -26,7 +26,7 @@ bSupportsIndexBufferUAVs=true
 bSupports4ComponentUAVReadWrite=true
 bSupportsManualVertexFetch=false
 bSupportsDualSourceBlending=true
-FriendlyName=LOCTEXT("FriendlyShaderPlatformName_TVOS_SM5", "TVOS Shader Model 5")
+FriendlyName=LOCTEXT("FriendlyShaderPlatformName_TVOS_SM5", "TVOS Metal SM5")
 
 [ShaderPlatform METAL_TVOS]
 Language=Metal
@@ -36,4 +36,4 @@ bIsMobile=true
 bTargetsTiledGPU=true
 bSupports4ComponentUAVReadWrite=true
 bSupportsManualVertexFetch=false
-FriendlyName=LOCTEXT("FriendlyShaderPlatformName_TVOS_ES31", "TVOS Shader Model 3_1")
\ No newline at end of file
+FriendlyName=LOCTEXT("FriendlyShaderPlatformName_TVOS_ES31", "TVOS Metal ES31")
\ No newline at end of file
diff --git a/Engine/Config/VulkanPC/DataDrivenPlatformInfo.ini b/Engine/Config/VulkanPC/DataDrivenPlatformInfo.ini
--- a/Engine/Config/VulkanPC/DataDrivenPlatformInfo.ini
+++ b/Engine/Config/VulkanPC/DataDrivenPlatformInfo.ini
@@ -131,7 +131,7 @@ bSupportsROV=false
 bSupportsOIT=false
 
 EnablesHLSL2021ByDefault=1
-FriendlyName=LOCTEXT("FriendlyShaderPlatformName_VulkanPC_SM5", "Vulkan PC Shader Model 5")
+FriendlyName=LOCTEXT("FriendlyShaderPlatformName_VulkanPC_SM5", "PC Vulkan SM5")
 
 BindlessSupport=AllShaderTypes
 
@@ -146,4 +146,4 @@ bSupportsMobileMultiView=true
 bSupportsIndexBufferUAVs=true
 bSupportsManualVertexFetch=false
 bSupportsSceneDataCompressedTransforms=true
-FriendlyName=LOCTEXT("FriendlyShaderPlatformName_VulkanPC_ES31", "Vulkan PC Shader Model 3_1")
+FriendlyName=LOCTEXT("FriendlyShaderPlatformName_VulkanPC_ES31", "PC Vulkan ES31")
diff --git a/Engine/Config/Windows/DataDrivenPlatformInfo.ini b/Engine/Config/Windows/DataDrivenPlatformInfo.ini
--- a/Engine/Config/Windows/DataDrivenPlatformInfo.ini
+++ b/Engine/Config/Windows/DataDrivenPlatformInfo.ini
@@ -75,7 +75,7 @@ bSupportsROV=true
 bSupportsOIT=true
 BindlessSupport=RayTracingOnly
 SupportsMultiViewport=RuntimeDependent
-FriendlyName=LOCTEXT("FriendlyShaderPlatformName_D3D_SM5", "D3D Shader Model 5")
+FriendlyName=LOCTEXT("FriendlyShaderPlatformName_D3D_SM5", "PC D3D SM5")
 bSupportsDualSourceBlending=true
 
 [PreviewPlatform PCD3D_SM5]
@@ -85,7 +85,7 @@ ShaderPlatform=PCD3D_SM5
 MenuTooltip=LOCTEXT("PreviewMenuTooltip_PCD3D_SM5", "DirectX 11, OpenGL 4.3+, PS4, XB1")
 
 [ShaderPlatform PCD3D_SM6]
-FriendlyName=LOCTEXT("FriendlyShaderPlatformName_D3D_SM6", "D3D Shader Model 6")
+FriendlyName=LOCTEXT("FriendlyShaderPlatformName_D3D_SM6", "PC D3D SM6")
 Language=D3D
 MaxFeatureLevel=SM6
 ShaderFormat=PCD3D_SM6
@@ -157,7 +157,7 @@ bSupportsVertexShaderLayer=true
 bSupportsManualVertexFetch=false
 bSupportsSceneDataCompressedTransforms=true
 bSupportsDebugViewShaders=true
-FriendlyName=LOCTEXT("FriendlyShaderPlatformName_D3D_ES31", "D3D Shader Model 3_1")
+FriendlyName=LOCTEXT("FriendlyShaderPlatformName_D3D_ES31", "PC D3D ES31")
 SupportsMultiViewport=RuntimeDependent
 
 [PreviewPlatform PCD3D_ES3_1]
@@ -174,4 +174,4 @@ MaxFeatureLevel=ES3_1
 ShaderFormat=GLSL_150_ES31
 bIsMobile=true
 bSupportsManualVertexFetch=false
-FriendlyName=LOCTEXT("FriendlyShaderPlatformName_PCGL_ES31", "Opengl PC Shader Model 3.1")
+FriendlyName=LOCTEXT("FriendlyShaderPlatformName_PCGL_ES31", "PC Opengl ES31")
