Set bSupportsClipDistance to false for OPENGL_PCES3_1 ShaderPlatform

diff --git a/Engine/Config/Windows/DataDrivenPlatformInfo.ini b/Engine/Config/Windows/DataDrivenPlatformInfo.ini
--- a/Engine/Config/Windows/DataDrivenPlatformInfo.ini
+++ b/Engine/Config/Windows/DataDrivenPlatformInfo.ini
@@ -172,4 +172,5 @@ MaxFeatureLevel=ES3_1
 ShaderFormat=GLSL_150_ES31
 bIsMobile=true
 bSupportsManualVertexFetch=false
+bSupportsClipDistance = false
 FriendlyName=LOCTEXT("FriendlyShaderPlatformName_PCGL_ES31", "PC Opengl ES31")
