Bump windows version from 18362 to 19041 to meet the minimum spec for GameInput API

diff --git a/Engine/Config/Windows/Windows_SDK.json b/Engine/Config/Windows/Windows_SDK.json
--- a/Engine/Config/Windows/Windows_SDK.json
+++ b/Engine/Config/Windows/Windows_SDK.json
@@ -1,5 +1,5 @@
 {
-"MainVersion" : "10.0.18362.0",
+"MainVersion" : "10.0.19041.0",
 "MinVersion" : "10.0.00000.0",
 "MaxVersion" : "10.9.99999.0",
 }
