- FIxed the Windows PreviewPlaforms to have the correct PlatformName

diff --git a/Engine/Config/VulkanPC/DataDrivenPlatformInfo.ini b/Engine/Config/VulkanPC/DataDrivenPlatformInfo.ini
--- a/Engine/Config/VulkanPC/DataDrivenPlatformInfo.ini
+++ b/Engine/Config/VulkanPC/DataDrivenPlatformInfo.ini
@@ -274,7 +274,7 @@ bSupportsUniformBufferObjects = true
 FriendlyName=LOCTEXT("FriendlyShaderPlatformName_VulkanPC_Mobile", "PC Vulkan Mobile")
 
 [PreviewPlatform VULKAN_SM5]
-PlatformName=VULKAN_SM5
+PlatformName=Linux
 ShaderFormat=SF_VULKAN_SM5
 ShaderPlatform=VULKAN_SM5
 MenuTooltip=LOCTEXT("PreviewMenuTooltip_VulkanPC_SM5", "Linux using SM5 profile")
diff --git a/Engine/Config/Windows/DataDrivenPlatformInfo.ini b/Engine/Config/Windows/DataDrivenPlatformInfo.ini
--- a/Engine/Config/Windows/DataDrivenPlatformInfo.ini
+++ b/Engine/Config/Windows/DataDrivenPlatformInfo.ini
@@ -79,7 +79,7 @@ bSupportsDualSourceBlending=true
 bSupportsNNEShaders=true
 
 [PreviewPlatform PCD3D_SM5]
-PlatformName=PCD3D_SM5
+PlatformName=Windows
 ShaderFormat=PCD3D_SM5
 ShaderPlatform=PCD3D_SM5
 MenuTooltip=LOCTEXT("PreviewMenuTooltip_PCD3D_SM5", "DirectX 11, OpenGL 4.3+, PS4, XB1")
@@ -140,7 +140,7 @@ bSupportsNNEShaders=true
 MaxSamplers=32
 
 [PreviewPlatform PCD3D_SM6]
-PlatformName=PCD3D_SM6
+PlatformName=Windows
 ShaderFormat=PCD3D_SM6
 ShaderPlatform=PCD3D_SM6
 MenuTooltip=LOCTEXT("PreviewMenuTooltip_PCD3D_SM6", "Windows using SM6 profile")
@@ -164,7 +164,7 @@ FriendlyName=LOCTEXT("FriendlyShaderPlatformName_D3D_Mobile", "PC D3D Mobile")
 SupportsMultiViewport=RuntimeDependent
 
 [PreviewPlatform PCD3D_ES3_1]
-PlatformName=PCD3D_ES3_1
+PlatformName=Windows
 ShaderFormat=PCD3D_ES31
 ShaderPlatform=PCD3D_ES3_1
 MenuTooltip=LOCTEXT("PreviewMenuTooltip_PCD3D_Mobile", "Windows using Mobile profile.")
