Support for MSVC 17.9.

diff --git a/Engine/Config/Windows/Windows_SDK.json b/Engine/Config/Windows/Windows_SDK.json
--- a/Engine/Config/Windows/Windows_SDK.json
+++ b/Engine/Config/Windows/Windows_SDK.json
@@ -20,6 +20,7 @@
 	],
 
 	"PreferredVisualCppVersions": [
+		"14.39.33519-14.39.99999", // VS2022 17.9.x
 		"14.38.33130-14.38.99999", // VS2022 17.8.x
 		"14.37.32822-14.37.99999", // VS2022 17.7.x
 		"14.36.32532-14.36.99999", // VS2022 17.6.x
