Enable Work Graph support in D3D12 SM6 now that we have updated DXC to support it.

diff --git a/Engine/Config/Windows/DataDrivenPlatformInfo.ini b/Engine/Config/Windows/DataDrivenPlatformInfo.ini
--- a/Engine/Config/Windows/DataDrivenPlatformInfo.ini
+++ b/Engine/Config/Windows/DataDrivenPlatformInfo.ini
@@ -138,8 +138,7 @@ bSupportsNNEShaders=true
 bSupportsIndependentSamplers=true
 MaxSamplers=32
 SupportsBarycentricsSemantic=RuntimeDependent
-;todo: Enable work graphs when DXC supports lib_6_8
-bSupportsWorkGraphs=false
+bSupportsWorkGraphs=true
 
 [PreviewPlatform PCD3D_SM6]
 PlatformName=PCD3D_SM6
