Enable r.Nanite.Streaming.ReservedResources on Windows

diff --git a/Engine/Config/Windows/WindowsEngine.ini b/Engine/Config/Windows/WindowsEngine.ini
--- a/Engine/Config/Windows/WindowsEngine.ini
+++ b/Engine/Config/Windows/WindowsEngine.ini
@@ -20,6 +20,7 @@ QualityLevelMapping="high"
 ; Enable single RHI thread stall for D3D11
 r.Occlusion.SingleRHIThreadStall=1
 r.Shadow.DetectVertexShaderLayerAtRuntime=1
+r.Nanite.Streaming.ReservedResources=1
 D3D12.Bindless.ResourceDescriptorHeapSize=100000
 D3D12.Bindless.SamplerDescriptorHeapSize=2048
 
