Disallow 14.39.33519 - 14.39.33522 toolchains

#rnx

diff --git a/Engine/Config/Windows/Windows_SDK.json b/Engine/Config/Windows/Windows_SDK.json
--- a/Engine/Config/Windows/Windows_SDK.json
+++ b/Engine/Config/Windows/Windows_SDK.json
@@ -11,7 +11,7 @@
 		"15.0.0-15.999",
 	],
 	"MinimumClangVersion": "15.0.0",
-	
+
 	// this is a strange format, but instead of a version range, this mapping min clang version for each MSVC version
 	"MinimumRequiredClangVersion": [
 		"14.37-16", // VS2022 17.7.x - 17.8.x
@@ -19,6 +19,7 @@
 		"14.34-14", // VS2022 17.4.x
 	],
 
+	// Version number is the MSVC family, which is the version in the Visual Studio folder
 	"PreferredVisualCppVersions": [
 		"14.38.33130-14.38.99999", // VS2022 17.8.x
 		"14.37.32822-14.37.99999", // VS2022 17.7.x
@@ -26,7 +27,9 @@
 		"14.35.32215-14.35.99999", // VS2022 17.5.x
 		"14.34.31933-14.34.99999", // VS2022 17.4.x
 	],
+
 	"BannedVisualCppVersions": [
+		"14.39.33519-14.39.33522" // VS2022 17.9.0 - 17.9.3: Internal compiler errors fixed in 17.9.4
 	],
 	"MinimumVisualCppVersion": "14.34.31933",
 
@@ -37,7 +40,7 @@
 
 	"VisualStudioSuggestedComponents": [
 		"Microsoft.Net.Component.4.6.2.TargetingPack",
-		"Microsoft.VisualStudio.Component.VC.14.38.17.8.x86.x64",
+		"Microsoft.VisualStudio.Component.VC.14.38.17.8.x86.x64", // Only even numbered LTS versions of MSVC should be suggested
 		"Microsoft.VisualStudio.Component.VC.Tools.x86.x64",
 		"Microsoft.VisualStudio.Component.Windows10SDK.22621",
 		"Microsoft.VisualStudio.Workload.CoreEditor",
