Enable PSO precaching for global compute shaders by default on Windows

diff --git a/Engine/Config/Windows/WindowsEngine.ini b/Engine/Config/Windows/WindowsEngine.ini
--- a/Engine/Config/Windows/WindowsEngine.ini
+++ b/Engine/Config/Windows/WindowsEngine.ini
@@ -23,6 +23,7 @@ r.Shadow.DetectVertexShaderLayerAtRuntime=1
 r.Nanite.Streaming.ReservedResources=1
 D3D12.Bindless.ResourceDescriptorHeapSize=100000
 D3D12.Bindless.SamplerDescriptorHeapSize=2048
+r.PSOPrecache.GlobalComputeShaders=1
 
 [PlatformCrypto]
 PlatformRequiresDataCrypto=True
