Add WindowsApplication KBM hardware device definition. Move the XInput hardware defintiion out to its own plugin as it should be.

#rnx

diff --git a/Engine/Config/Windows/WindowsInput.ini b/Engine/Config/Windows/WindowsInput.ini
--- a/Engine/Config/Windows/WindowsInput.ini
+++ b/Engine/Config/Windows/WindowsInput.ini
@@ -1,2 +1,2 @@
 [InputPlatformSettings_Windows InputPlatformSettings]
-+HardwareDevices=(InputClassName="XInputInterface",HardwareDeviceIdentifier="XInputController",PrimaryDeviceType=Gamepad,SupportedFeaturesMask=260)
\ No newline at end of file
++HardwareDevices=(InputClassName="WindowsApplication",HardwareDeviceIdentifier="KBM",PrimaryDeviceType=KeyboardAndMouse,SupportedFeaturesMask=3)
\ No newline at end of file
diff --git a/Engine/Plugins/Runtime/Windows/XInputDevice/Config/Input.ini b/Engine/Plugins/Runtime/Windows/XInputDevice/Config/Input.ini
new file mode 100644
--- /dev/null
+++ b/Engine/Plugins/Runtime/Windows/XInputDevice/Config/Input.ini
@@ -0,0 +1,2 @@
+[InputPlatformSettings_Windows InputPlatformSettings]
++HardwareDevices=(InputClassName="XInputInterface",HardwareDeviceIdentifier="XInputController",PrimaryDeviceType=Gamepad,SupportedFeaturesMask=260)
\ No newline at end of file
