Add msvc 14.39.x to preferred msvc toolchains

diff --git a/Engine/Config/Windows/Windows_SDK.json b/Engine/Config/Windows/Windows_SDK.json
--- a/Engine/Config/Windows/Windows_SDK.json
+++ b/Engine/Config/Windows/Windows_SDK.json
@@ -8,24 +8,25 @@
 	"PreferredClangVersions": [
 		// Clang 17: [REDACTED_PROFILE_URL]/llvm-project/issues/71976 should not be used as a preferred version until this issue is resolved and the max affected version is added to ClangWarnings.cs
 		"16.0.0-16.999",
-		"15.0.0-15.999",
+		"15.0.0-15.999"
 	],
 	"MinimumClangVersion": "15.0.0",
 
 	// this is a strange format, but instead of a version range, this mapping min clang version for each MSVC version
 	"MinimumRequiredClangVersion": [
-		"14.37-16", // VS2022 17.7.x - 17.8.x
+		"14.37-16", // VS2022 17.7.x - 17.9.x
 		"14.35-15", // VS2022 17.5.x - 17.6.x
-		"14.34-14", // VS2022 17.4.x
+		"14.34-14" // VS2022 17.4.x
 	],
 
 	// Version number is the MSVC family, which is the version in the Visual Studio folder
 	"PreferredVisualCppVersions": [
+		"14.39.33519-14.39.99999", // VS2022 17.9.x (17.9.0 - 17.9.3 is banned)
 		"14.38.33130-14.38.99999", // VS2022 17.8.x
 		"14.37.32822-14.37.99999", // VS2022 17.7.x
 		"14.36.32532-14.36.99999", // VS2022 17.6.x
 		"14.35.32215-14.35.99999", // VS2022 17.5.x
-		"14.34.31933-14.34.99999", // VS2022 17.4.x
+		"14.34.31933-14.34.99999" // VS2022 17.4.x
 	],
 
 	"BannedVisualCppVersions": [
@@ -40,7 +41,7 @@
 
 	"VisualStudioSuggestedComponents": [
 		"Microsoft.Net.Component.4.6.2.TargetingPack",
-		"Microsoft.VisualStudio.Component.VC.14.38.17.8.x86.x64", // Only even numbered LTS versions of MSVC should be suggested
+		"Microsoft.VisualStudio.Component.VC.14.38.17.8.x86.x64", // Only LTSC versions of MSVC should be suggested
 		"Microsoft.VisualStudio.Component.VC.Tools.x86.x64",
 		"Microsoft.VisualStudio.Component.Windows10SDK.22621",
 		"Microsoft.VisualStudio.Workload.CoreEditor",
