Enable r.RayTracing.EnableOnDemand by default on Windows.

diff --git a/Engine/Config/Windows/WindowsEngine.ini b/Engine/Config/Windows/WindowsEngine.ini
--- a/Engine/Config/Windows/WindowsEngine.ini
+++ b/Engine/Config/Windows/WindowsEngine.ini
@@ -24,9 +24,11 @@ r.Nanite.Streaming.ReservedResources=1
 D3D12.Bindless.ResourceDescriptorHeapSize=100000
 D3D12.Bindless.SamplerDescriptorHeapSize=2048
 r.PSOPrecache.GlobalComputeShaders=1
+; Allow ray tracing to be toggled using r.RayTracing.Enable=1/0 without restarting
+r.RayTracing.EnableOnDemand=True
 
 [PlatformCrypto]
 PlatformRequiresDataCrypto=True
 
 [/Script/Engine.RendererSettings]
-r.Shaders.RemoveUnusedInterpolators=1
\ No newline at end of file
+r.Shaders.RemoveUnusedInterpolators=1
