Add ATLMFC to suggested components

#rnx

diff --git a/Engine/Config/Windows/Windows_SDK.json b/Engine/Config/Windows/Windows_SDK.json
--- a/Engine/Config/Windows/Windows_SDK.json
+++ b/Engine/Config/Windows/Windows_SDK.json
@@ -41,6 +41,8 @@
 	"VisualStudioSuggestedComponents": [
 		"Microsoft.Net.Component.4.6.2.TargetingPack",
 		"Microsoft.VisualStudio.Component.VC.14.38.17.8.x86.x64", // Only LTSC versions of MSVC should be suggested
+		"Microsoft.VisualStudio.Component.VC.14.38.17.8.ATL", // Match version of MSVC
+		"Microsoft.VisualStudio.Component.VC.14.38.17.8.MFC", // Match version of MSVC
 		"Microsoft.VisualStudio.Component.VC.Tools.x86.x64",
 		"Microsoft.VisualStudio.Component.Windows10SDK.22621",
 		"Microsoft.VisualStudio.Workload.CoreEditor",
