Ban all 14.39.x toolchains due to codegen issue when compiling for AVX arch

diff --git a/Engine/Config/Windows/Windows_SDK.json b/Engine/Config/Windows/Windows_SDK.json
--- a/Engine/Config/Windows/Windows_SDK.json
+++ b/Engine/Config/Windows/Windows_SDK.json
@@ -29,7 +29,7 @@
 	],
 
 	"BannedVisualCppVersions": [
-		"14.39.33519-14.39.33522" // VS2022 17.9.0 - 17.9.3: Internal compiler errors fixed in 17.9.4
+		"14.39.33519-14.39.99999" // VS2022 17.9.0 - 17.9.3: Internal compiler errors. 17.9.x: Codegen issue causes runtime crash for AVX/2/512
 	],
 	"MinimumVisualCppVersion": "14.34.31933",
 
