Enable shader warnings as errors by default on Windows platforms

This only applies to global shaders for now and is done only for one platform to gauge the impact and minimize disruption for platforms which not everyone has installed or can test easily.

diff --git a/Engine/Config/Windows/BaseWindowsEngine.ini b/Engine/Config/Windows/BaseWindowsEngine.ini
--- a/Engine/Config/Windows/BaseWindowsEngine.ini
+++ b/Engine/Config/Windows/BaseWindowsEngine.ini
@@ -35,3 +35,5 @@ MaxNumThreadsWithTlsSlots=512
 [ConsoleVariables]
 ; necessary for LowMemoryWarning events to be fired
 memory.MemoryPressureCriticalThresholdMB=512
+; Enable warnings as errors for global shaders only
+r.Shaders.WarningsAsErrors=1
