UnrealBuildTool: Update minimum clang to 16 in to match already enforced minimum from SDK config file

#rnx

diff --git a/Engine/Config/Windows/Windows_SDK.json b/Engine/Config/Windows/Windows_SDK.json
--- a/Engine/Config/Windows/Windows_SDK.json
+++ b/Engine/Config/Windows/Windows_SDK.json
@@ -7,10 +7,9 @@
 
 	"PreferredClangVersions": [
 		// Clang 17: [REDACTED_PROFILE_URL]/llvm-project/issues/71976 should not be used as a preferred version until this issue is resolved and the max affected version is added to ClangWarnings.cs
-		"16.0.0-16.999",
-		"15.0.0-15.999"
+		"16.0.0-16.999"
 	],
-	"MinimumClangVersion": "15.0.0",
+	"MinimumClangVersion": "16.0.0",
 
 	// this is a strange format, but instead of a version range, this mapping min clang version for each MSVC version
 	"MinimumRequiredClangVersion": [
@@ -43,7 +42,7 @@
 		"Microsoft.VisualStudio.Component.VC.14.38.17.8.x86.x64", // Only LTSC versions of MSVC should be suggested
 		"Microsoft.VisualStudio.Component.VC.14.38.17.8.ATL", // Match version of MSVC
 		"Microsoft.VisualStudio.Component.VC.Tools.x86.x64",
-		"Microsoft.VisualStudio.Component.Windows10SDK.22621",
+		"Microsoft.VisualStudio.Component.Windows11SDK.22621",
 		"Microsoft.VisualStudio.Workload.CoreEditor",
 		"Microsoft.VisualStudio.Workload.ManagedDesktop",
 		"Microsoft.VisualStudio.Workload.NativeDesktop",
diff --git a/Engine/Source/Programs/UnrealBuildTool/Platform/Linux/LinuxToolChain.cs b/Engine/Source/Programs/UnrealBuildTool/Platform/Linux/LinuxToolChain.cs
--- a/Engine/Source/Programs/UnrealBuildTool/Platform/Linux/LinuxToolChain.cs
+++ b/Engine/Source/Programs/UnrealBuildTool/Platform/Linux/LinuxToolChain.cs
@@ -55,10 +55,10 @@ namespace UnrealBuildTool
 			: this(UnrealTargetPlatform.Linux, InArchitecture, InSDK, InOptions, InLogger)
 		{
 			// prevent unknown clangs since the build is likely to fail on too old or too new compilers
-			if (CompilerVersionLessThan(15, 0, 0) || CompilerVersionGreaterOrEqual(17, 0, 0))
+			if (CompilerVersionLessThan(16, 0, 0) || CompilerVersionGreaterOrEqual(17, 0, 0))
 			{
 				throw new BuildException(
-					String.Format("This version of the Unreal Engine can only be compiled with clang 15.x or 16.x. clang {0} may not build it - please use a different version.",
+					String.Format("This version of the Unreal Engine can only be compiled with clang 16.x. clang {0} may not build it - please use a different version.",
 						Info.ClangVersion)
 					);
 			}
