MSVC 14.40 requires at least Clang 17

#rnx

diff --git a/Engine/Config/Windows/Windows_SDK.json b/Engine/Config/Windows/Windows_SDK.json
--- a/Engine/Config/Windows/Windows_SDK.json
+++ b/Engine/Config/Windows/Windows_SDK.json
@@ -13,6 +13,7 @@
 
 	// this is a strange format, but instead of a version range, this mapping min clang version for each MSVC version
 	"MinimumRequiredClangVersion": [
+		"14.40-17", // VS2022 17.10.x
 		"14.37-16", // VS2022 17.7.x - 17.9.x
 		"14.35-15", // VS2022 17.5.x - 17.6.x
 		"14.34-14" // VS2022 17.4.x
