Enable bindless raytracing only by default on PCD3D_SM6

diff --git a/Engine/Config/Windows/BaseWindowsEngine.ini b/Engine/Config/Windows/BaseWindowsEngine.ini
--- a/Engine/Config/Windows/BaseWindowsEngine.ini
+++ b/Engine/Config/Windows/BaseWindowsEngine.ini
@@ -15,8 +15,8 @@ BindlessResources=Disabled
 BindlessSamplers=Disabled
 
 [PCD3D_SM6]
-BindlessResources=Disabled
-BindlessSamplers=Disabled
+BindlessResources=RayTracingOnly
+BindlessSamplers=RayTracingOnly
 
 [SF_VULKAN_SM5]
 BindlessResources=Disabled
