Allow Clang 18, prioritized below 16 until the preferred update is complete

diff --git a/Engine/Config/Windows/Windows_SDK.json b/Engine/Config/Windows/Windows_SDK.json
--- a/Engine/Config/Windows/Windows_SDK.json
+++ b/Engine/Config/Windows/Windows_SDK.json
@@ -6,8 +6,8 @@
 	"MinSoftwareVersion": "10.0.18362.0",
 
 	"PreferredClangVersions": [
-		// Clang 17: [REDACTED_PROFILE_URL]/llvm-project/issues/71976 should not be used as a preferred version until this issue is resolved and the max affected version is added to ClangWarnings.cs
-		"16.0.0-16.999"
+		"16.0.0-16.999",
+		"18.1.3-18.999"
 	],
 	"MinimumClangVersion": "16.0.0",
 
@@ -29,6 +29,7 @@
 	],
 
 	"BannedVisualCppVersions": [
+		"14.40.33807-14.40.33810", // VS2022 17.10.0: Codegen issue causes runtime crash for AVX/2/512
 		"14.39.33519-14.39.99999" // VS2022 17.9.0 - 17.9.3: Internal compiler errors. 17.9.x: Codegen issue causes runtime crash for AVX/2/512
 	],
 	"MinimumVisualCppVersion": "14.34.31933",
