Enable ShaderBindingLayout by default on D3D12 PC

diff --git a/Engine/Config/Windows/DataDrivenPlatformInfo.ini b/Engine/Config/Windows/DataDrivenPlatformInfo.ini
--- a/Engine/Config/Windows/DataDrivenPlatformInfo.ini
+++ b/Engine/Config/Windows/DataDrivenPlatformInfo.ini
@@ -70,6 +70,7 @@ bSupportsIndexBufferUAVs=true
 bSupportsROV=true
 bSupportsOIT=true
 BindlessSupport=RayTracingOnly
+StaticShaderBindingLayoutSupport=RayTracingOnly
 SupportsMultiViewport=RuntimeDependent
 FriendlyName=LOCTEXT("FriendlyShaderPlatformName_D3D_SM5", "PC D3D SM5")
 bSupportsDualSourceBlending=true
@@ -129,6 +130,7 @@ bSupportsCapsuleShadows=true
 bSupportsPercentageCloserShadows=true
 bSupportsIndexBufferUAVs=true
 BindlessSupport=AllShaderTypes
+StaticShaderBindingLayoutSupport=RayTracingOnly
 bSupportsROV=true
 bSupportsOIT=true
 bSupportsRealTypes=RuntimeDependent
