Add VS .uproject support to recommended components

#rnx

diff --git a/Engine/Config/Windows/Windows_SDK.json b/Engine/Config/Windows/Windows_SDK.json
--- a/Engine/Config/Windows/Windows_SDK.json
+++ b/Engine/Config/Windows/Windows_SDK.json
@@ -41,6 +41,7 @@
 
 	"VisualStudioSuggestedComponents": [
 		"Microsoft.Net.Component.4.6.2.TargetingPack",
+		"Microsoft.VisualStudio.Component.Unreal.Workspace",
 		"Microsoft.VisualStudio.Component.VC.14.38.17.8.x86.x64", // Only LTSC versions of MSVC should be suggested
 		"Microsoft.VisualStudio.Component.VC.14.38.17.8.ATL", // Match version of MSVC
 		"Microsoft.VisualStudio.Component.VC.Tools.x86.x64",
