Update MSVC preferred version to 14.40, bump preferred clang version to 18.1.3

diff --git a/Engine/Config/Windows/Windows_SDK.json b/Engine/Config/Windows/Windows_SDK.json
--- a/Engine/Config/Windows/Windows_SDK.json
+++ b/Engine/Config/Windows/Windows_SDK.json
@@ -6,8 +6,8 @@
 	"MinSoftwareVersion": "10.0.18362.0",
 
 	"PreferredClangVersions": [
-		"16.0.0-16.999",
-		"18.1.3-18.999"
+		"18.1.3-18.999",
+		"16.0.0-16.999"
 	],
 	"MinimumClangVersion": "16.0.0",
 
@@ -21,11 +21,13 @@
 
 	// Version number is the MSVC family, which is the version in the Visual Studio folder
 	"PreferredVisualCppVersions": [
+		"14.40.33807-14.40.99999", // VS2022 17.10.x
 		"14.38.33130-14.38.99999", // VS2022 17.8.x
-		"14.37.32822-14.37.99999", // VS2022 17.7.x
 		"14.36.32532-14.36.99999", // VS2022 17.6.x
-		"14.35.32215-14.35.99999", // VS2022 17.5.x
-		"14.34.31933-14.34.99999" // VS2022 17.4.x
+		"14.34.31933-14.34.99999", // VS2022 17.4.x
+		// Non-LTSC versions
+		"14.37.32822-14.37.99999", // VS2022 17.7.x
+		"14.35.32215-14.35.99999" // VS2022 17.5.x
 	],
 
 	"BannedVisualCppVersions": [
@@ -42,8 +44,8 @@
 	"VisualStudioSuggestedComponents": [
 		"Microsoft.Net.Component.4.6.2.TargetingPack",
 		"Microsoft.VisualStudio.Component.Unreal.Workspace",
-		"Microsoft.VisualStudio.Component.VC.14.38.17.8.x86.x64", // Only LTSC versions of MSVC should be suggested
-		"Microsoft.VisualStudio.Component.VC.14.38.17.8.ATL", // Match version of MSVC
+		"Microsoft.VisualStudio.Component.VC.14.40.17.10.x86.x64", // Only LTSC versions of MSVC should be suggested
+		"Microsoft.VisualStudio.Component.VC.14.40.17.10.ATL", // Match version of MSVC
 		"Microsoft.VisualStudio.Component.VC.Tools.x86.x64",
 		"Microsoft.VisualStudio.Component.Windows11SDK.22621",
 		"Microsoft.VisualStudio.Workload.CoreEditor",
