Update preferred clang version for Windows platforms to [REDACTED_PHONE]

diff --git a/Engine/Config/Windows/Windows_SDK.json b/Engine/Config/Windows/Windows_SDK.json
--- a/Engine/Config/Windows/Windows_SDK.json
+++ b/Engine/Config/Windows/Windows_SDK.json
@@ -6,8 +6,8 @@
 	"MinSoftwareVersion": "10.0.18362.0",
 
 	"PreferredClangVersions": [
-		"16.0.0-16.999",
-		"18.1.3-18.999"
+		"18.1.3-18.999",
+		"16.0.0-16.999"
 	],
 	"MinimumClangVersion": "16.0.0",
 
