[RHI] Disable vertex shader layer for preview platform if the runtime platform does not support it.

This also disables bSupportsVertexShaderLayer for PCD3D_SM6 and PCD3D_ES3_1 since D3D does not support vertex shader layer, it is primarily used by Metal to compensate for the lack of geometry shaders.
This fixes a crash that occurs due to a VS/PS interpolator mismatch at PSO creation time when switching to Mac Metal SM5/6 preview with the D3D runtime.

#rnx

diff --git a/Engine/Config/Windows/DataDrivenPlatformInfo.ini b/Engine/Config/Windows/DataDrivenPlatformInfo.ini
--- a/Engine/Config/Windows/DataDrivenPlatformInfo.ini
+++ b/Engine/Config/Windows/DataDrivenPlatformInfo.ini
@@ -124,7 +124,7 @@ bSupportsMeshShadersTier1=true
 bSupportsMeshShadersWithClipDistance=true
 bSupportsHairStrandGeometry=true
 bSupportsFFTBloom=true
-bSupportsVertexShaderLayer=true
+bSupportsVertexShaderLayer=false
 bSupportsInstancedStereo=true
 bSupportsCapsuleShadows=true
 bSupportsPercentageCloserShadows=true
@@ -160,7 +160,7 @@ bSupportsDistanceFields = false
 bSupportsComputeFramework=true
 bSupportsByteBufferComputeShaders = true
 bSupportsInstancedStereo=true
-bSupportsVertexShaderLayer=true
+bSupportsVertexShaderLayer=false
 bSupportsManualVertexFetch=false
 bSupportsSceneDataCompressedTransforms=true
 bSupportsDebugViewShaders=true
diff --git a/Engine/Source/Runtime/RHI/Private/DataDrivenShaderPlatformInfo.cpp b/Engine/Source/Runtime/RHI/Private/DataDrivenShaderPlatformInfo.cpp
--- a/Engine/Source/Runtime/RHI/Private/DataDrivenShaderPlatformInfo.cpp
+++ b/Engine/Source/Runtime/RHI/Private/DataDrivenShaderPlatformInfo.cpp
@@ -492,6 +492,7 @@ void FGenericDataDrivenShaderPlatformInfo::UpdatePreviewPlatforms()
 				else
 				{
 					PREVIEW_DISABLE_IF_RUNTIME_UNSUPPORTED(bSupportsInstancedStereo);
+					PREVIEW_DISABLE_IF_RUNTIME_UNSUPPORTED(bSupportsVertexShaderLayer);
 				}
 
 				// Settings that should be kept true if the runtime also supports it.
