Enable Nanite software VRS (contrast adaptive) on all appropriate platforms by default now

diff --git a/Engine/Config/Windows/WindowsEngine.ini b/Engine/Config/Windows/WindowsEngine.ini
--- a/Engine/Config/Windows/WindowsEngine.ini
+++ b/Engine/Config/Windows/WindowsEngine.ini
@@ -24,6 +24,8 @@ D3D12.Bindless.ResourceDescriptorHeapSize=32768
 D3D12.Bindless.SamplerDescriptorHeapSize=2048
 r.PSOPrecache.GlobalComputeShaders=1
 r.DynamicRes.DynamicFrameTime=1
+r.VRS.EnableSoftware=1
+r.VRS.ContrastAdaptiveShading=1
 
 [PlatformCrypto]
 PlatformRequiresDataCrypto=True
