Update minimum MSVC to 14.38.33130 & minimum Visual Studio 2022 to 17.8

diff --git a/Engine/Config/Windows/Windows_SDK.json b/Engine/Config/Windows/Windows_SDK.json
--- a/Engine/Config/Windows/Windows_SDK.json
+++ b/Engine/Config/Windows/Windows_SDK.json
@@ -15,26 +15,19 @@
 	"MinimumRequiredClangVersion": [
 		"14.40-17", // VS2022 17.10.x
 		"14.37-16", // VS2022 17.7.x - 17.9.x
-		"14.35-15", // VS2022 17.5.x - 17.6.x
-		"14.34-14" // VS2022 17.4.x
 	],
 
 	// Version number is the MSVC family, which is the version in the Visual Studio folder
 	"PreferredVisualCppVersions": [
 		"14.40.33807-14.40.99999", // VS2022 17.10.x
-		"14.38.33130-14.38.99999", // VS2022 17.8.x
-		"14.36.32532-14.36.99999", // VS2022 17.6.x
-		"14.34.31933-14.34.99999", // VS2022 17.4.x
-		// Non-LTSC versions
-		"14.37.32822-14.37.99999", // VS2022 17.7.x
-		"14.35.32215-14.35.99999" // VS2022 17.5.x
+		"14.38.33130-14.38.99999" // VS2022 17.8.x
 	],
 
 	"BannedVisualCppVersions": [
 		"14.40.33807-14.40.33810", // VS2022 17.10.0: Codegen issue causes runtime crash for AVX/2/512
 		"14.39.33519-14.39.99999" // VS2022 17.9.0 - 17.9.3: Internal compiler errors. 17.9.x: Codegen issue causes runtime crash for AVX/2/512
 	],
-	"MinimumVisualCppVersion": "14.34.31933",
+	"MinimumVisualCppVersion": "14.38.33130",
 
 	"PreferredIntelOneApiVersions": [
 		"2024.1.0-2024.9999"
@@ -61,7 +54,7 @@
 	"VisualStudio2022SuggestedComponents": [
 	],
 
-	"MinimumVisualStudio2022Version": "17.4",
+	"MinimumVisualStudio2022Version": "17.8",
 
 	"MinimumWindowsX64TargetVersion": "0x601", // Win7
 	"MinimumWindowsArm64TargetVersion":  "0xA00" // Win10
