Bump minimum Windows SDK to 10.0.19041.0

#rnx

diff --git a/Engine/Config/Windows/Windows_SDK.json b/Engine/Config/Windows/Windows_SDK.json
--- a/Engine/Config/Windows/Windows_SDK.json
+++ b/Engine/Config/Windows/Windows_SDK.json
@@ -1,9 +1,9 @@
 {
 	"MainVersion": "10.0.19041.0",
-	"MinVersion": "10.0.18362.0",
+	"MinVersion": "10.0.19041.0",
 	"MaxVersion": "10.9.99999.0",
 
-	"MinSoftwareVersion": "10.0.18362.0",
+	"MinSoftwareVersion": "10.0.19041.0",
 
 	"PreferredClangVersions": [
 		"18.1.3-18.999",
@@ -14,7 +14,7 @@
 	// this is a strange format, but instead of a version range, this mapping min clang version for each MSVC version
 	"MinimumRequiredClangVersion": [
 		"14.40-17", // VS2022 17.10.x
-		"14.37-16", // VS2022 17.7.x - 17.9.x
+		"14.37-16" // VS2022 17.7.x - 17.9.x
 	],
 
 	// Version number is the MSVC family, which is the version in the Visual Studio folder
@@ -57,6 +57,6 @@
 	"MinimumVisualStudio2022Version": "17.8",
 
 	"MinimumWindowsX64TargetVersion": "0x601", // Win7
-	"MinimumWindowsArm64TargetVersion":  "0xA00" // Win10
+	"MinimumWindowsArm64TargetVersion": "0xA00" // Win10
 
 }
