Downgrade preferred MSVC due to LTCG codegen causing runtime crash

#rnx

diff --git a/Engine/Config/Windows/Windows_SDK.json b/Engine/Config/Windows/Windows_SDK.json
--- a/Engine/Config/Windows/Windows_SDK.json
+++ b/Engine/Config/Windows/Windows_SDK.json
@@ -19,12 +19,11 @@
 
 	// Version number is the MSVC family, which is the version in the Visual Studio folder
 	"PreferredVisualCppVersions": [
-		"14.40.33807-14.40.99999", // VS2022 17.10.x
 		"14.38.33130-14.38.99999" // VS2022 17.8.x
 	],
 
 	"BannedVisualCppVersions": [
-		"14.40.33807-14.40.33810", // VS2022 17.10.0: Codegen issue causes runtime crash for AVX/2/512
+		"14.40.33807-14.40.99999", // VS2022 17.10.0: Codegen issue causes runtime crash for AVX/2/512. 17.10.x: PGO optimize issue causes runtime crash
 		"14.39.33519-14.39.99999" // VS2022 17.9.0 - 17.9.3: Internal compiler errors. 17.9.x: Codegen issue causes runtime crash for AVX/2/512
 	],
 	"MinimumVisualCppVersion": "14.38.33130",
@@ -37,8 +36,8 @@
 	"VisualStudioSuggestedComponents": [
 		"Microsoft.Net.Component.4.6.2.TargetingPack",
 		"Microsoft.VisualStudio.Component.Unreal.Workspace",
-		"Microsoft.VisualStudio.Component.VC.14.40.17.10.x86.x64", // Only LTSC versions of MSVC should be suggested
-		"Microsoft.VisualStudio.Component.VC.14.40.17.10.ATL", // Match version of MSVC
+		"Microsoft.VisualStudio.Component.VC.14.38.17.8.x86.x64", // Only LTSC versions of MSVC should be suggested
+		"Microsoft.VisualStudio.Component.VC.14.38.17.8.ATL", // Match version of MSVC
 		"Microsoft.VisualStudio.Component.VC.Tools.x86.x64",
 		"Microsoft.VisualStudio.Component.Windows11SDK.22621",
 		"Microsoft.VisualStudio.Workload.CoreEditor",
