Bump preferred WindowsSDK to 10.0.22621.0

#rnx

diff --git a/Engine/Config/Windows/Windows_SDK.json b/Engine/Config/Windows/Windows_SDK.json
--- a/Engine/Config/Windows/Windows_SDK.json
+++ b/Engine/Config/Windows/Windows_SDK.json
@@ -1,5 +1,5 @@
 {
-	"MainVersion": "10.0.19041.0",
+	"MainVersion": "10.0.22621.0",
 	"MinVersion": "10.0.19041.0",
 	"MaxVersion": "10.9.99999.0",
 
