Disable bindless for raytracing due to visual corruption

diff --git a/Engine/Config/Windows/BaseWindowsEngine.ini b/Engine/Config/Windows/BaseWindowsEngine.ini
--- a/Engine/Config/Windows/BaseWindowsEngine.ini
+++ b/Engine/Config/Windows/BaseWindowsEngine.ini
@@ -16,8 +16,8 @@ BindlessSamplers=Disabled
 MaxClosuresPerPixel=1
 
 [PCD3D_SM6]
-BindlessResources=RayTracingOnly
-BindlessSamplers=RayTracingOnly
+BindlessResources=Disabled
+BindlessSamplers=Disabled
 
 [SF_VULKAN_SM5]
 BindlessResources=Disabled
diff --git a/Engine/Config/Windows/DataDrivenPlatformInfo.ini b/Engine/Config/Windows/DataDrivenPlatformInfo.ini
--- a/Engine/Config/Windows/DataDrivenPlatformInfo.ini
+++ b/Engine/Config/Windows/DataDrivenPlatformInfo.ini
@@ -122,7 +122,7 @@ bSupportsInstancedStereo=true
 bSupportsPercentageCloserShadows=true
 bSupportsIndexBufferUAVs=true
 BindlessSupport=AllShaderTypes
-StaticShaderBindingLayoutSupport=RayTracingOnly
+StaticShaderBindingLayoutSupport=Unsupported
 bSupportsROV=true
 bSupportsOIT=true
 bSupportsRealTypes=RuntimeDependent
