Disabled the critical memory reporting threshold on desktop until we have a more reliable method.

diff --git a/Engine/Config/Linux/BaseLinuxEngine.ini b/Engine/Config/Linux/BaseLinuxEngine.ini
--- a/Engine/Config/Linux/BaseLinuxEngine.ini
+++ b/Engine/Config/Linux/BaseLinuxEngine.ini
@@ -11,7 +11,3 @@ bAgreeToCrashUpload=false
 ; Bindless is necessary for Vulkan ray tracing
 BindlessResources=Enabled
 BindlessSamplers=Enabled
-
-[ConsoleVariables]
-; necessary for LowMemoryWarning events to be fired
-memory.MemoryPressureCriticalThresholdMB=512
diff --git a/Engine/Config/LinuxArm64/BaseLinuxArm64Engine.ini b/Engine/Config/LinuxArm64/BaseLinuxArm64Engine.ini
--- a/Engine/Config/LinuxArm64/BaseLinuxArm64Engine.ini
+++ b/Engine/Config/LinuxArm64/BaseLinuxArm64Engine.ini
@@ -1,6 +1,2 @@
 [CrashReportClient]
 bAgreeToCrashUpload=false
-
-[ConsoleVariables]
-; necessary for LowMemoryWarning events to be fired
-memory.MemoryPressureCriticalThresholdMB=512
diff --git a/Engine/Config/Mac/BaseMacEngine.ini b/Engine/Config/Mac/BaseMacEngine.ini
--- a/Engine/Config/Mac/BaseMacEngine.ini
+++ b/Engine/Config/Mac/BaseMacEngine.ini
@@ -36,9 +36,6 @@ QualityLevelMapping="high"
 r.HairStrands.Visibility.MSAA.SamplePerPixel=1
 r.SceneCulling.ExplicitCellBounds=1
 
-; necessary for LowMemoryWarning events to be fired
-memory.MemoryPressureCriticalThresholdMB=512
-
 [SF_METAL_SM6]
 BindlessResources=Enabled
 BindlessSamplers=Enabled
diff --git a/Engine/Config/Windows/BaseWindowsEngine.ini b/Engine/Config/Windows/BaseWindowsEngine.ini
--- a/Engine/Config/Windows/BaseWindowsEngine.ini
+++ b/Engine/Config/Windows/BaseWindowsEngine.ini
@@ -29,7 +29,3 @@ BindlessSamplers=Enabled
 [/Script/WindowsRuntimeSettings.WindowsRuntimeSettings]
 MaxNumTlsSlots=2048
 MaxNumThreadsWithTlsSlots=512
-
-[ConsoleVariables]
-; necessary for LowMemoryWarning events to be fired
-memory.MemoryPressureCriticalThresholdMB=512
