- Made the SDK.json files be standards-compliant json (mostly turning comments into KV pairs)

diff --git a/Engine/Config/Android/Android_SDK.json b/Engine/Config/Android/Android_SDK.json
--- a/Engine/Config/Android/Android_SDK.json
+++ b/Engine/Config/Android/Android_SDK.json
@@ -8,5 +8,5 @@
 	"platforms": "android-34",
 	"build-tools": "34.0.0",
 	"cmake": "3.22.1",
-	"ndk": "25.1.8937393",
+	"ndk": "25.1.8937393"
 }
diff --git a/Engine/Config/Apple/Apple_SDK.json b/Engine/Config/Apple/Apple_SDK.json
--- a/Engine/Config/Apple/Apple_SDK.json
+++ b/Engine/Config/Apple/Apple_SDK.json
@@ -1,23 +1,24 @@
 {
-	// Xcode versions
-	"MainVersion": "15.2",
-	"MinVersion": "15.2.0",
-	"MaxVersion": "16.9.0",
-	// !!!
-	// NOTE: If you update the MaxVersion, double check the AppleVersionToLLVMVersion array below!!!
-	// !!!
+	"//1": "Xcode versions:",
+		"MainVersion": "15.2",
+		"MinVersion": "15.2.0",
+		"MaxVersion": "16.9.0",
+	"//2": "!!!",
+	"//3": "NOTE: If you update the MaxVersion, double check the AppleVersionToLLVMVersion array below!!!",
+	"//4": "!!!",
 
-	// The versions on Windows are iTunes versions
-	"MinVersion_Win64": "1100.0.0.0",
-	"MaxVersion_Win64": "8999.0",
+	"//5": "The versions on Windows are iTunes versions:",
+		"MinVersion_Win64": "1100.0.0.0",
+		"MaxVersion_Win64": "8999.0",
 
-	// This is not a version range, but a mapping of Xcode clang versions to LLVM versions, for shared version checks with other clangs
-	// Version mapping can be found at https://en.wikipedia.org/wiki/Xcode#Toolchain_versions
-	// The first version is the first version that is using the second version source LLVM
-	"AppleVersionToLLVMVersions": [
-		"14.0.0-14.0.0", // Xcode 14.0 -> LLVM 14.0
-		"14.0.3-15.0.0", // Xcode 14.3 -> LLVM 15.0
-		"15.0.0-16.0.0", // Xcode 15.0 -> LLVM 16.0
-		"16.0.0-17.0.6", // Xcode 16.0 -> LLVM 17.0.6
-	],
+	"//6": "This is not a version range, but a mapping of Xcode clang versions to LLVM versions, for shared version checks with other clangs",
+	"//7": "Version mapping can be found at https://en.wikipedia.org/wiki/Xcode#Toolchain_versions",
+	"//8": "The first half of the pair is the first version that is using the second version source LLVM",
+	"//9": "For instance, Xcode 16 uses LLVM 17.0.6",
+		"AppleVersionToLLVMVersions": [
+			"14.0.0-14.0.0",
+			"14.0.3-15.0.0",
+			"15.0.0-16.0.0",
+			"16.0.0-17.0.6"
+		]
 }
diff --git a/Engine/Config/IOS/IOS_SDK.json b/Engine/Config/IOS/IOS_SDK.json
--- a/Engine/Config/IOS/IOS_SDK.json
+++ b/Engine/Config/IOS/IOS_SDK.json
@@ -1,5 +1,5 @@
 {
 	"ParentSDKFile": "../Apple/Apple_SDK.json",
 
-	"MinSoftwareVersion":  "15.0",
+	"MinSoftwareVersion":  "15.0"
 }
diff --git a/Engine/Config/Linux/Linux_SDK.json b/Engine/Config/Linux/Linux_SDK.json
--- a/Engine/Config/Linux/Linux_SDK.json
+++ b/Engine/Config/Linux/Linux_SDK.json
@@ -1,7 +1,7 @@
 {
-"MainVersion" : "v23_clang-18.1.0-rockylinux8",
-"MinVersion" : "v23_clang-18.1.0-rockylinux8",
-"MaxVersion" : "v23_clang-18.1.0-rockylinux8",
+	"MainVersion" : "v23_clang-18.1.0-rockylinux8",
+	"MinVersion" : "v23_clang-18.1.0-rockylinux8",
+	"MaxVersion" : "v23_clang-18.1.0-rockylinux8",
 
-"AutoSDKPlatform" : "Linux_x64",
+	"AutoSDKPlatform" : "Linux_x64"
 }
diff --git a/Engine/Config/Mac/Mac_SDK.json b/Engine/Config/Mac/Mac_SDK.json
--- a/Engine/Config/Mac/Mac_SDK.json
+++ b/Engine/Config/Mac/Mac_SDK.json
@@ -7,6 +7,6 @@
 	"DeploymentTarget": "13.00",
 	"BuildTarget": "13.00",
 
-	// macOS Ventura 13.0 (Released October 24, 2022)
+	"//": "macOS Ventura 13.0 (Released October 24, 2022)",
 	"MinSoftwareVersion": "13.00"
 }
diff --git a/Engine/Config/TVOS/TVOS_SDK.json b/Engine/Config/TVOS/TVOS_SDK.json
--- a/Engine/Config/TVOS/TVOS_SDK.json
+++ b/Engine/Config/TVOS/TVOS_SDK.json
@@ -1,4 +1,4 @@
 {
-	// inherit the IOS MinSoftwareVersion
+	"//": "inherit the IOS MinSoftwareVersion",
 	"ParentSDKFile": "../IOS/IOS_SDK.json"
 }
diff --git a/Engine/Config/Windows/Windows_SDK.json b/Engine/Config/Windows/Windows_SDK.json
--- a/Engine/Config/Windows/Windows_SDK.json
+++ b/Engine/Config/Windows/Windows_SDK.json
@@ -11,20 +11,25 @@
 	],
 	"MinimumClangVersion": "16.0.0",
 
-	// this is a strange format, but instead of a version range, this mapping min clang version for each MSVC version
+	"//01": "this is a strange format, but instead of a version range, this maps min clang version for each MSVC version",
+	"//02": "14.40 = VS2022 17.10.x -> clang 17",
+	"//03": "14.37 = VS2022 17.7.x-17.9.x -> clang 16",
 	"MinimumRequiredClangVersion": [
-		"14.40-17", // VS2022 17.10.x
-		"14.37-16" // VS2022 17.7.x - 17.9.x
+		"14.40-17",
+		"14.37-16"
 	],
 
-	// Version number is the MSVC family, which is the version in the Visual Studio folder
+	"//04": "Version number is the MSVC family, which is the version in the Visual Studio folder",
+	"//05": "14.38.33130-14.38.99999 = VS2022 17.8.x",
 	"PreferredVisualCppVersions": [
-		"14.38.33130-14.38.99999" // VS2022 17.8.x
+		"14.38.33130-14.38.99999"
 	],
 
+	"//06": "14.40.x -- VS2022 17.10.0: Codegen issue causes runtime crash for AVX/2/512. 17.10.x: PGO optimize issue causes runtime crash",
+	"//07": "14.39.x -- VS2022 17.9.0 - 17.9.3: Internal compiler errors. 17.9.x: Codegen issue causes runtime crash for AVX/2/512",
 	"BannedVisualCppVersions": [
-		"14.40.33807-14.40.99999", // VS2022 17.10.0: Codegen issue causes runtime crash for AVX/2/512. 17.10.x: PGO optimize issue causes runtime crash
-		"14.39.33519-14.39.99999" // VS2022 17.9.0 - 17.9.3: Internal compiler errors. 17.9.x: Codegen issue causes runtime crash for AVX/2/512
+		"14.40.33807-14.40.99999",
+		"14.39.33519-14.39.99999"
 	],
 	"MinimumVisualCppVersion": "14.38.33130",
 
@@ -33,11 +38,13 @@
 	],
 	"MinimumIntelOneApiVersion": "2025.0.0",
 
+	"//08": "Microsoft.VisualStudio.Component.VC.14.38.17.8.x86.x64 -- Only LTSC versions of MSVC should be suggested",
+	"//09": "Microsoft.VisualStudio.Component.VC.14.38.17.8.ATL -- Match version of MSVC",
 	"VisualStudioSuggestedComponents": [
 		"Microsoft.Net.Component.4.6.2.TargetingPack",
 		"Microsoft.VisualStudio.Component.Unreal.Workspace",
-		"Microsoft.VisualStudio.Component.VC.14.38.17.8.x86.x64", // Only LTSC versions of MSVC should be suggested
-		"Microsoft.VisualStudio.Component.VC.14.38.17.8.ATL", // Match version of MSVC
+		"Microsoft.VisualStudio.Component.VC.14.38.17.8.x86.x64",
+		"Microsoft.VisualStudio.Component.VC.14.38.17.8.ATL",
 		"Microsoft.VisualStudio.Component.VC.Tools.x86.x64",
 		"Microsoft.VisualStudio.Component.Windows11SDK.22621",
 		"Microsoft.VisualStudio.Workload.CoreEditor",
@@ -55,7 +62,8 @@
 
 	"MinimumVisualStudio2022Version": "17.8",
 
+	"//10": "X64 -- 0x601 = Win7",
+	"//11": "Arm64 -- 0xA00 = Win10",
 	"MinimumWindowsX64TargetVersion": "0x601", // Win7
 	"MinimumWindowsArm64TargetVersion": "0xA00" // Win10
-
 }
diff --git a/Engine/Platforms/VisionOS/Config/VisionOS_SDK.json b/Engine/Platforms/VisionOS/Config/VisionOS_SDK.json
--- a/Engine/Platforms/VisionOS/Config/VisionOS_SDK.json
+++ b/Engine/Platforms/VisionOS/Config/VisionOS_SDK.json
@@ -6,5 +6,5 @@
 	"MinVersion": "15.3.0",
 	
 	// runtime OS version
-	"MinSoftwareVersion":  "1.0",
+	"MinSoftwareVersion":  "1.0"
 }
