Enable inline raytracing on d3d12 win64

diff --git a/Engine/Config/Windows/DataDrivenPlatformInfo.ini b/Engine/Config/Windows/DataDrivenPlatformInfo.ini
--- a/Engine/Config/Windows/DataDrivenPlatformInfo.ini
+++ b/Engine/Config/Windows/DataDrivenPlatformInfo.ini
@@ -85,7 +85,7 @@ MaxMeshShaderThreadGroupSize=128
 bSupportsDebugViewShaders=true
 bSupportsRayTracing=true
 bSupportsRayTracingShaders=true
-bSupportsInlineRayTracing=false
+bSupportsInlineRayTracing=true
 bSupportsRayTracingCallableShaders=true
 bSupportsRayTracingProceduralPrimitive=true
 bSupportsRayTracingIndirectInstanceData=true
