Bump minimum clang to 18

#rnx

diff --git a/Engine/Config/Windows/Windows_SDK.json b/Engine/Config/Windows/Windows_SDK.json
--- a/Engine/Config/Windows/Windows_SDK.json
+++ b/Engine/Config/Windows/Windows_SDK.json
@@ -6,27 +6,28 @@
 	"MinSoftwareVersion": "10.0.19041.0",
 
 	"PreferredClangVersions": [
-		"18.1.3-18.999",
-		"16.0.0-16.999"
+		"18.1.3-18.999"
 	],
-	"MinimumClangVersion": "16.0.0",
+	"MinimumClangVersion": "18.1.3",
 
-	"//01": "this is a strange format, but instead of a version range, this maps min clang version for each MSVC version",
-	"//02": "14.40 = VS2022 17.10.x -> clang 17",
-	"//03": "14.37 = VS2022 17.7.x-17.9.x -> clang 16",
+	"//01.1": "this is a strange format, but instead of a version range, this maps min clang version for each MSVC version",
+	"//01.2": "14.42 = VS2022 17.12.x -> clang 18",
+	"//01.3": "14.40 = VS2022 17.10.x -> clang 17",
+	"//01.4": "14.37 = VS2022 17.7.x-17.9.x -> clang 16",
 	"MinimumRequiredClangVersion": [
+		"14.42-18",
 		"14.40-17",
 		"14.37-16"
 	],
 
-	"//04": "Version number is the MSVC family, which is the version in the Visual Studio folder",
-	"//05": "14.38.33130-14.38.99999 = VS2022 17.8.x",
+	"//02.1": "Version number is the MSVC family, which is the version in the Visual Studio folder",
+	"//02.2": "14.38.33130-14.38.99999 = VS2022 17.8.x",
 	"PreferredVisualCppVersions": [
 		"14.38.33130-14.38.99999"
 	],
 
-	"//06": "14.40.x -- VS2022 17.10.0: Codegen issue causes runtime crash for AVX/2/512. 17.10.x: PGO optimize issue causes runtime crash",
-	"//07": "14.39.x -- VS2022 17.9.0 - 17.9.3: Internal compiler errors. 17.9.x: Codegen issue causes runtime crash for AVX/2/512",
+	"//03.1": "14.40.x -- VS2022 17.10.0: Codegen issue causes runtime crash for AVX/2/512. 17.10.x: PGO optimize issue causes runtime crash",
+	"//03.2": "14.39.x -- VS2022 17.9.0 - 17.9.3: Internal compiler errors. 17.9.x: Codegen issue causes runtime crash for AVX/2/512",
 	"BannedVisualCppVersions": [
 		"14.40.33807-14.40.99999",
 		"14.39.33519-14.39.99999"
@@ -38,8 +39,8 @@
 	],
 	"MinimumIntelOneApiVersion": "2025.0.0",
 
-	"//08": "Microsoft.VisualStudio.Component.VC.14.38.17.8.x86.x64 -- Only LTSC versions of MSVC should be suggested",
-	"//09": "Microsoft.VisualStudio.Component.VC.14.38.17.8.ATL -- Match version of MSVC",
+	"//04.1": "Microsoft.VisualStudio.Component.VC.14.38.17.8.x86.x64 -- Only LTSC versions of MSVC should be suggested",
+	"//04.2": "Microsoft.VisualStudio.Component.VC.14.38.17.8.ATL -- Match version of MSVC",
 	"VisualStudioSuggestedComponents": [
 		"Microsoft.Net.Component.4.6.2.TargetingPack",
 		"Microsoft.VisualStudio.Component.Unreal.Workspace",
@@ -47,6 +48,7 @@
 		"Microsoft.VisualStudio.Component.VC.14.38.17.8.ATL",
 		"Microsoft.VisualStudio.Component.VC.Tools.x86.x64",
 		"Microsoft.VisualStudio.Component.Windows11SDK.22621",
+		"Microsoft.VisualStudio.Component.VC.Llvm.Clang",
 		"Microsoft.VisualStudio.Workload.CoreEditor",
 		"Microsoft.VisualStudio.Workload.ManagedDesktop",
 		"Microsoft.VisualStudio.Workload.NativeDesktop",
@@ -62,8 +64,8 @@
 
 	"MinimumVisualStudio2022Version": "17.8",
 
-	"//10": "X64 -- 0x601 = Win7",
-	"//11": "Arm64 -- 0xA00 = Win10",
+	"//05.1": "X64 -- 0x601 = Win7",
+	"//05.2": "Arm64 -- 0xA00 = Win10",
 	"MinimumWindowsX64TargetVersion": "0x601", // Win7
 	"MinimumWindowsArm64TargetVersion": "0xA00" // Win10
 }
