Add clang 19 to preferred compiler list in preparation for next update, currently lower priority than clang 18

#rnx

diff --git a/Engine/Config/Windows/Windows_SDK.json b/Engine/Config/Windows/Windows_SDK.json
--- a/Engine/Config/Windows/Windows_SDK.json
+++ b/Engine/Config/Windows/Windows_SDK.json
@@ -6,7 +6,8 @@
 	"MinSoftwareVersion": "10.0.19041.0",
 
 	"PreferredClangVersions": [
-		"18.1.3-18.999"
+		"18.1.3-18.999",
+		"19.1.0-19.999"
 	],
 	"MinimumClangVersion": "18.1.3",
 
