Windows: migrate INI settings to BaseWindowsEngine.ini from WindowsEngine.ini

diff --git a/Engine/Config/Windows/BaseWindowsEngine.ini b/Engine/Config/Windows/BaseWindowsEngine.ini
--- a/Engine/Config/Windows/BaseWindowsEngine.ini
+++ b/Engine/Config/Windows/BaseWindowsEngine.ini
@@ -1,5 +1,39 @@
 [Audio]
+; Defining below allows switching to audio mixer using -audiomixer commandline
 AudioCaptureModuleName=AudioCaptureWasapi
+AudioMixerModuleName=AudioMixerXAudio2
+PlatformFormat=OGG
+PlatformStreamingFormat=OGG
+
+; Defines a platform-specific volume headroom (in dB) for audio to provide better platform consistency with respect to volume levels.
+PlatformHeadroomDB=-3
+
+[TextureStreaming]
+; PoolSizeVRAMPercentage is how much percentage of GPU Dedicated VRAM should be used as a TexturePool cache for streaming textures (0 - unlimited streaming)
+PoolSizeVRAMPercentage=70
+
+[SystemSettings]
+r.setres=1280x720
+framepro.ScopeMinTimeMicroseconds=10
+fx.NiagaraAllowRuntimeScalabilityChanges=1
+;Perplatform to PerQualityLevel conversion mapping for platform groups
+QualityLevelMapping="high"
+; Enable single RHI thread stall for D3D11
+r.Occlusion.SingleRHIThreadStall=1
+r.Nanite.Streaming.ReservedResources=1
+D3D12.Bindless.ResourceDescriptorHeapSize=32768
+D3D12.Bindless.SamplerDescriptorHeapSize=2048
+r.PSOPrecache.GlobalShaders=1
+r.DynamicRes.DynamicFrameTime=1
+r.VRS.EnableSoftware=1
+r.VRS.ContrastAdaptiveShading=1
+
+[PlatformCrypto]
+PlatformRequiresDataCrypto=True
+
+[/Script/Engine.RendererSettings]
+r.Shaders.RemoveUnusedInterpolators=1
+r.Shadow.DetectVertexShaderLayerAtRuntime=1
 
 [OnlineSubsystem]
 LocalPlatformName=WIN
diff --git a/Engine/Config/Windows/WindowsEngine.ini b/Engine/Config/Windows/WindowsEngine.ini
deleted file mode 100644
