MSVC 14.44 requires Clang 19

#rnx

diff --git a/Engine/Config/Windows/Windows_SDK.json b/Engine/Config/Windows/Windows_SDK.json
--- a/Engine/Config/Windows/Windows_SDK.json
+++ b/Engine/Config/Windows/Windows_SDK.json
@@ -12,10 +12,12 @@
 	"MinimumClangVersion": "18.1.3",
 
 	"//01.1": "this is a strange format, but instead of a version range, this maps min clang version for each MSVC version",
-	"//01.2": "14.42 = VS2022 17.12.x -> clang 18",
-	"//01.3": "14.40 = VS2022 17.10.x -> clang 17",
-	"//01.4": "14.37 = VS2022 17.7.x-17.9.x -> clang 16",
+	"//01.2": "14.44 = VS2022 17.14.x -> clang 19",
+	"//01.3": "14.42 = VS2022 17.12.x -> clang 18",
+	"//01.4": "14.40 = VS2022 17.10.x -> clang 17",
+	"//01.5": "14.37 = VS2022 17.7.x-17.9.x -> clang 16",
 	"MinimumRequiredClangVersion": [
+		"14.44-19",
 		"14.42-18",
 		"14.40-17",
 		"14.37-16"
