Set WindowsApplication.UseWorkerThreadForRawInput to False until support is added for right-button primary

diff --git a/Engine/Config/Windows/BaseWindowsEngine.ini b/Engine/Config/Windows/BaseWindowsEngine.ini
--- a/Engine/Config/Windows/BaseWindowsEngine.ini
+++ b/Engine/Config/Windows/BaseWindowsEngine.ini
@@ -60,4 +60,4 @@ MaxNumTlsSlots=2048
 MaxNumThreadsWithTlsSlots=512
 
 [ConsoleVariables]
-WindowsApplication.UseWorkerThreadForRawInput=true
\ No newline at end of file
+WindowsApplication.UseWorkerThreadForRawInput=False
