Substrate - settings setup for 60hz performance first. Allowing the substrate gbuffer on pc and high end consoles only.
Renamed shader platform data MaxClosuresPerPixel to SubstrateMaxClosuresPerPixel.

#tests pc editor dx12, preview for multiple platforms.
#rnx

diff --git a/Engine/Config/Android/BaseAndroidEngine.ini b/Engine/Config/Android/BaseAndroidEngine.ini
--- a/Engine/Config/Android/BaseAndroidEngine.ini
+++ b/Engine/Config/Android/BaseAndroidEngine.ini
@@ -23,12 +23,6 @@ gc.MaxObjectsInGame=131072
 NumPreallocatedVertices=200
 
 [/Script/Engine.RendererSettings]
-r.Substrate.RoughDiffuse=0
-r.Substrate.ShadingQuality=2
-r.Substrate.TileCoord8bits=1
-r.Substrate.SheenQuality=2
-r.Substrate.Glints=0
-r.Substrate.ClosuresPerPixel=1
 ; Volumetric Clouds are not supported by mobile renderer, skip materials/global shaders compilation.
 r.VolumetricCloud.Support=0
 
diff --git a/Engine/Config/BaseDeviceProfiles.ini b/Engine/Config/BaseDeviceProfiles.ini
--- a/Engine/Config/BaseDeviceProfiles.ini
+++ b/Engine/Config/BaseDeviceProfiles.ini
@@ -1394,19 +1394,7 @@ BaseProfileName=Android_Mid
 [Windows_Preview_ES31 DeviceProfile]
 DeviceType=Windows
 BaseProfileName=Windows
-+CVars=r.Substrate.RoughDiffuse=0
-+CVars=r.Substrate.ShadingQuality=2
-+CVars=r.Substrate.TileCoord8bits=1
-+CVars=r.Substrate.SheenQuality=2
-+CVars=r.Substrate.Glints=0
-+CVars=r.Substrate.ClosuresPerPixel=1
 
 [Windows_Preview_ES31_SDF DeviceProfile]
 DeviceType=Windows
 BaseProfileName=Windows
-+CVars=r.Substrate.RoughDiffuse=0
-+CVars=r.Substrate.ShadingQuality=2
-+CVars=r.Substrate.TileCoord8bits=1
-+CVars=r.Substrate.SheenQuality=2
-+CVars=r.Substrate.Glints=0
-+CVars=r.Substrate.ClosuresPerPixel=1
diff --git a/Engine/Config/BaseEngine.ini b/Engine/Config/BaseEngine.ini
--- a/Engine/Config/BaseEngine.ini
+++ b/Engine/Config/BaseEngine.ini
@@ -3663,6 +3663,17 @@ bRequireDeviceSupportsARKit=true
 r.GPUCrashDebugging=false
 MaxSkinBones=(Default=65536,PerPlatform=(("Mobile", 256)))
 
+; Substrate settings for all platforms by default are setup for 60hz performance 
+r.Substrate.RoughDiffuse=0						; Lower lighting cost: disable rough diffuse
+r.Substrate.SheenQuality=2						; Lower lighting cost: simpler fuzz/cloth shading
+r.Substrate.ShadingQuality=2					; Lower lighting cost: simpler GGX, diffuse, non chromatic energy presevation
+r.Substrate.TileCoord8bits=1					; Memory and performance optimisation: 8 bit tile coord for lower resolution
+r.Substrate.Glints=0							; GBuffer and performance optimisation: disable GLints
+r.Substrate.ClosuresPerPixel=1					; GBuffer and performance optimisation: single closure max per pixel
+r.Substrate.OpaqueMaterialRoughRefraction=0		; Memory and performance optimisation: disable opaque rough refraction of multi layered materials
+r.Refraction.Blur=0								; Memory and performance optimisation: disable translucenct rough refraction
+r.Substrate.BlendableGBuffer=1					; GBuffer decal optimisation: use the legacy blendable gbuffer to avoid dbuffer decal memory overhead
+
 [Messaging]
 bAllowDelayedMessaging=false
 
diff --git a/Engine/Config/IOS/BaseIOSEngine.ini b/Engine/Config/IOS/BaseIOSEngine.ini
--- a/Engine/Config/IOS/BaseIOSEngine.ini
+++ b/Engine/Config/IOS/BaseIOSEngine.ini
@@ -11,12 +11,6 @@ gc.MaxObjectsInGame=131072
 [/Script/Engine.RendererSettings]
 r.DistanceFields=0
 r.TemporalAA.Mobile.UseCompute=0
-r.Substrate.RoughDiffuse=0
-r.Substrate.ShadingQuality=2
-r.Substrate.TileCoord8bits=1
-r.Substrate.SheenQuality=2
-r.Substrate.Glints=0
-r.Substrate.ClosuresPerPixel=1
 ; Disables parallel GDME by default due to not enough threads
 r.Visibility.DynamicMeshElements.Parallel=0
 ; Volumetric Clouds are not supported by mobile renderer, skip materials/global shaders compilation.
diff --git a/Engine/Config/Linux/BaseLinuxEngine.ini b/Engine/Config/Linux/BaseLinuxEngine.ini
--- a/Engine/Config/Linux/BaseLinuxEngine.ini
+++ b/Engine/Config/Linux/BaseLinuxEngine.ini
@@ -10,3 +10,15 @@ bAgreeToCrashUpload=false
 [ShaderPlatformConfig VULKAN_SM6]
 ; Bindless is necessary for Vulkan ray tracing
 BindlessConfiguration=All
+
+[SystemSettings]
+; Substrate settings for more advanced shading on PC and editor by default.
+r.Substrate.RoughDiffuse=1
+r.Substrate.SheenQuality=1
+r.Substrate.ShadingQuality=1
+r.Substrate.TileCoord8bits=0
+r.Substrate.Glints=1
+r.Substrate.ClosuresPerPixel=4
+r.Substrate.OpaqueMaterialRoughRefraction=0
+r.Refraction.Blur=1
+r.Substrate.BlendableGBuffer=0
diff --git a/Engine/Config/LinuxArm64/BaseLinuxArm64Engine.ini b/Engine/Config/LinuxArm64/BaseLinuxArm64Engine.ini
--- a/Engine/Config/LinuxArm64/BaseLinuxArm64Engine.ini
+++ b/Engine/Config/LinuxArm64/BaseLinuxArm64Engine.ini
@@ -1,2 +1,15 @@
 [CrashReportClient]
 bAgreeToCrashUpload=false
+
+
+[SystemSettings]
+; Substrate settings for more advanced shading on PC and editor by default.
+r.Substrate.RoughDiffuse=1
+r.Substrate.SheenQuality=1
+r.Substrate.ShadingQuality=1
+r.Substrate.TileCoord8bits=0
+r.Substrate.Glints=1
+r.Substrate.ClosuresPerPixel=4
+r.Substrate.OpaqueMaterialRoughRefraction=0
+r.Refraction.Blur=1
+r.Substrate.BlendableGBuffer=0
diff --git a/Engine/Config/Windows/BaseWindowsEngine.ini b/Engine/Config/Windows/BaseWindowsEngine.ini
--- a/Engine/Config/Windows/BaseWindowsEngine.ini
+++ b/Engine/Config/Windows/BaseWindowsEngine.ini
@@ -16,7 +16,7 @@ PoolSizeVRAMPercentage=70
 r.setres=1280x720
 framepro.ScopeMinTimeMicroseconds=10
 fx.NiagaraAllowRuntimeScalabilityChanges=1
-;Perplatform to PerQualityLevel conversion mapping for platform groups
+; Perplatform to PerQualityLevel conversion mapping for platform groups
 QualityLevelMapping="high"
 ; Enable single RHI thread stall for D3D11
 r.Occlusion.SingleRHIThreadStall=1
@@ -26,6 +26,16 @@ D3D12.Bindless.SamplerDescriptorHeapSize=2048
 r.PSOPrecache.GlobalShaders=1
 r.VRS.EnableSoftware=1
 r.VRS.ContrastAdaptiveShading=1
+; Substrate settings for more advanced shading on PC and editor by default.
+r.Substrate.RoughDiffuse=1
+r.Substrate.SheenQuality=1
+r.Substrate.ShadingQuality=1
+r.Substrate.TileCoord8bits=1
+r.Substrate.Glints=1
+r.Substrate.ClosuresPerPixel=4
+r.Substrate.OpaqueMaterialRoughRefraction=0
+r.Refraction.Blur=1
+r.Substrate.BlendableGBuffer=0
 
 [PlatformCrypto]
 PlatformRequiresDataCrypto=True
@@ -42,14 +52,14 @@ DeviceProfileSelectionModule="WindowsDeviceProfileSelector"
 
 [ShaderPlatformConfig PCD3D_SM5]
 BindlessConfiguration=Disabled
-MaxClosuresPerPixel=1
+SubstrateMaxClosuresPerPixel=1
 
 [ShaderPlatformConfig PCD3D_SM6]
 BindlessConfiguration=RayTracing
 
 [ShaderPlatformConfig VULKAN_SM5]
 BindlessConfiguration=Disabled
-MaxClosuresPerPixel=1
+SubstrateMaxClosuresPerPixel=1
 
 [ShaderPlatformConfig VULKAN_SM6]
 ; Bindless is necessary for Vulkan ray tracing
diff --git a/Engine/Source/Runtime/RenderCore/Private/RenderUtils.cpp b/Engine/Source/Runtime/RenderCore/Private/RenderUtils.cpp
--- a/Engine/Source/Runtime/RenderCore/Private/RenderUtils.cpp
+++ b/Engine/Source/Runtime/RenderCore/Private/RenderUtils.cpp
@@ -2013,11 +2013,11 @@ namespace Substrate
 
 		// Variant for shader compilation per platform
 		const TCHAR* CVarName = TEXT("r.Substrate.ClosuresPerPixel");
-		const TCHAR* ClosureName = TEXT("MaxClosuresPerPixel");
+		const TCHAR* ClosureName = TEXT("SubstrateMaxClosuresPerPixel");
 		static FShaderPlatformCachedIniValue<int32> CVarClosureBudget(CVarName);
 		int32 OutClosurePerPixel = CVarClosureBudget.Get(InPlatform);
 
-		// Override r.Substrate.ClosuresPerPixel with ShaderFormat/MaxClosuresPerPixel for supporting sm5/sm6 differences
+		// Override r.Substrate.ClosuresPerPixel with ShaderFormat/SubstrateMaxClosuresPerPixel for supporting sm5/sm6 differences
 		#if WITH_EDITORONLY_DATA
 		const FName ShaderFormatName = LegacyShaderPlatformToShaderFormat(InPlatform);
 		if (ITargetPlatform* TargetPlatform = GetTargetPlatformManager()->FindTargetPlatformWithSupport(TEXT("ShaderFormat"), ShaderFormatName))
