Set WindowsApplication.UseWorkerThreadForRawInput=False by default

Seeing crashes when WindowsApplication.UseWorkerThreadForRawInput=True.

diff --git a/Engine/Config/Windows/BaseWindowsEngine.ini b/Engine/Config/Windows/BaseWindowsEngine.ini
--- a/Engine/Config/Windows/BaseWindowsEngine.ini
+++ b/Engine/Config/Windows/BaseWindowsEngine.ini
@@ -70,4 +70,4 @@ MaxNumTlsSlots=2048
 MaxNumThreadsWithTlsSlots=512
 
 [ConsoleVariables]
-WindowsApplication.UseWorkerThreadForRawInput=True
+WindowsApplication.UseWorkerThreadForRawInput=False
