Change Substrate 8bit tile value to be 0 on PC.

diff --git a/Engine/Config/Windows/BaseWindowsEngine.ini b/Engine/Config/Windows/BaseWindowsEngine.ini
--- a/Engine/Config/Windows/BaseWindowsEngine.ini
+++ b/Engine/Config/Windows/BaseWindowsEngine.ini
@@ -30,7 +30,7 @@ r.VRS.ContrastAdaptiveShading=1
 r.Substrate.RoughDiffuse=1
 r.Substrate.SheenQuality=1
 r.Substrate.ShadingQuality=1
-r.Substrate.TileCoord8bits=1
+r.Substrate.TileCoord8bits=0
 r.Substrate.Glints=1
 r.Substrate.ClosuresPerPixel=4
 r.Substrate.OpaqueMaterialRoughRefraction=0
