Swap from XAudio2 to WASAPI backend

diff --git a/Engine/Config/Windows/BaseWindowsEngine.ini b/Engine/Config/Windows/BaseWindowsEngine.ini
--- a/Engine/Config/Windows/BaseWindowsEngine.ini
+++ b/Engine/Config/Windows/BaseWindowsEngine.ini
@@ -1,7 +1,7 @@
 [Audio]
 ; Defining below allows switching to audio mixer using -audiomixer commandline
 AudioCaptureModuleName=AudioCaptureWasapi
-AudioMixerModuleName=AudioMixerXAudio2
+AudioMixerModuleName=AudioMixerWasapi
 PlatformFormat=BINKA
 PlatformStreamingFormat=BINKA
 
