Make XAudio2 the default audio backend

#rnx

diff --git a/Engine/Config/Windows/BaseWindowsEngine.ini b/Engine/Config/Windows/BaseWindowsEngine.ini
--- a/Engine/Config/Windows/BaseWindowsEngine.ini
+++ b/Engine/Config/Windows/BaseWindowsEngine.ini
@@ -1,7 +1,7 @@
 [Audio]
 ; Defining below allows switching to audio mixer using -audiomixer commandline
 AudioCaptureModuleName=AudioCaptureWasapi
-AudioMixerModuleName=AudioMixerWasapi
+AudioMixerModuleName=AudioMixerXAudio2
 PlatformFormat=BINKA
 PlatformStreamingFormat=BINKA
 
