Add clang 20 to preferred list, still lower priority than clang 18 until we're ready to update

#rnx

diff --git a/Engine/Config/Windows/Windows_SDK.json b/Engine/Config/Windows/Windows_SDK.json
--- a/Engine/Config/Windows/Windows_SDK.json
+++ b/Engine/Config/Windows/Windows_SDK.json
@@ -7,6 +7,7 @@
 
 	"PreferredClangVersions": [
 		"18.1.3-18.999",
+		"20.1.0-20.999",
 		"19.1.0-19.999"
 	],
 	"MinimumClangVersion": "18.1.3",
