MSVC preferred version to 14.44.35211 (family version 14.44.35207)

diff --git a/Engine/Config/Windows/Windows_SDK.json b/Engine/Config/Windows/Windows_SDK.json
--- a/Engine/Config/Windows/Windows_SDK.json
+++ b/Engine/Config/Windows/Windows_SDK.json
@@ -25,16 +25,20 @@
 	],
 
 	"//02.1": "Version number is the MSVC family, which is the version in the Visual Studio folder",
-	"//02.2": "14.38.33130-14.38.99999 = VS2022 17.8.x",
+	"//02.2": "14.44.35207-14.44.99999 = VS2022 17.14.x",
+	"//02.3": "14.38.33130-14.38.99999 = VS2022 17.8.x",
 	"PreferredVisualCppVersions": [
+		"14.44.35207-14.44.99999",
 		"14.38.33130-14.38.99999"
 	],
 
-	"//03.1": "14.40.x -- VS2022 17.10.0: Codegen issue causes runtime crash for AVX/2/512. 17.10.x: PGO optimize issue causes runtime crash",
-	"//03.2": "14.39.x -- VS2022 17.9.0 - 17.9.3: Internal compiler errors. 17.9.x: Codegen issue causes runtime crash for AVX/2/512",
+	"//03.1": "14.44.35207 -- VS2022 17.14.0 - 17.4.3: Template compile error. Resolved with 14.44.35211",
+	"//03.2": "14.40.x -- VS2022 17.10.0: Codegen issue causes runtime crash for AVX/2/512. 17.10.x: PGO optimize issue causes runtime crash",
+	"//03.3": "14.39.x -- VS2022 17.9.0 - 17.9.3: Internal compiler errors. 17.9.x: Codegen issue causes runtime crash for AVX/2/512",
 	"BannedVisualCppVersions": [
-		"14.40.33807-14.40.99999",
-		"14.39.33519-14.39.99999"
+		"14.44.0-14.44.35210",
+		"14.40.0-14.43.99999",
+		"14.39.0-14.39.99999"
 	],
 	"MinimumVisualCppVersion": "14.38.33130",
 
@@ -43,10 +47,12 @@
 	],
 	"MinimumIntelOneApiVersion": "2025.1.0",
 
-	"//04.1": "Microsoft.VisualStudio.Component.VC.14.38.17.8.x86.x64 -- Only LTSC versions of MSVC should be suggested",
-	"//04.2": "Microsoft.VisualStudio.Component.VC.14.38.17.8.ATL -- Match version of MSVC",
+	"//04.1": "Microsoft.VisualStudio.Component.VC.14.44.17.14.x86.x64 -- Only LTSC versions of MSVC should be suggested",
+	"//04.2": "Microsoft.VisualStudio.Component.VC.14.44.17.14.ATL -- Match version of MSVC",
 	"VisualStudioSuggestedComponents": [
 		"Microsoft.Net.Component.4.6.2.TargetingPack",
+		"Microsoft.VisualStudio.Component.VC.14.44.17.14.x86.x64",
+		"Microsoft.VisualStudio.Component.VC.14.44.17.14.ATL",
 		"Microsoft.VisualStudio.Component.VC.14.38.17.8.x86.x64",
 		"Microsoft.VisualStudio.Component.VC.14.38.17.8.ATL",
 		"Microsoft.VisualStudio.Component.VC.Tools.x86.x64",
