Add GDK to the Windows restricted folders
- this is to support PC GDK
- still requires NDA platform access for now

#rnx

diff --git a/Engine/Config/Windows/DataDrivenPlatformInfo.ini b/Engine/Config/Windows/DataDrivenPlatformInfo.ini
--- a/Engine/Config/Windows/DataDrivenPlatformInfo.ini
+++ b/Engine/Config/Windows/DataDrivenPlatformInfo.ini
@@ -2,6 +2,7 @@
 bIsConfidential=false
 IniParent=VulkanPC
 PlatformNameAliases=Win64
+AdditionalRestrictedFolders=GDK
 TargetSettingsIniSectionName=/Script/WindowsTargetPlatform.WindowsTargetSettings
 Freezing_b32Bit=false
 Freezing_MaxFieldAlignment32=4
