Adjust clang version list to make clear what the preferred latest version is

#rnx

diff --git a/Engine/Config/Windows/Windows_SDK.json b/Engine/Config/Windows/Windows_SDK.json
--- a/Engine/Config/Windows/Windows_SDK.json
+++ b/Engine/Config/Windows/Windows_SDK.json
@@ -6,7 +6,8 @@
 	"MinSoftwareVersion": "10.0.19041.0",
 
 	"PreferredClangVersions": [
-		"20.1.0-20.999",
+		"20.1.8-20.999",
+		"20.1.0-20.1.7",
 		"19.1.0-19.999",
 		"18.1.3-18.999"
 	],
