Remove MSVC 14.38 from the preferred compiler list

#rnx

diff --git a/Engine/Config/Windows/Windows_SDK.json b/Engine/Config/Windows/Windows_SDK.json
--- a/Engine/Config/Windows/Windows_SDK.json
+++ b/Engine/Config/Windows/Windows_SDK.json
@@ -27,10 +27,8 @@
 
 	"//02.1": "Version number is the MSVC family, which is the version in the Visual Studio folder",
 	"//02.2": "14.44.35207-14.44.99999 = VS2022 17.14.x",
-	"//02.3": "14.38.33130-14.38.99999 = VS2022 17.8.x",
 	"PreferredVisualCppVersions": [
-		"14.44.35207-14.44.99999",
-		"14.38.33130-14.38.99999"
+		"14.44.35207-14.44.99999"
 	],
 
 	"//03.1": "14.44.35207 -- VS2022 17.14.0 - 17.4.3: Template compile error. Resolved with 14.44.35211",
