Update minimum clang to 18.1.8

#rnx

diff --git a/Engine/Config/Windows/Windows_SDK.json b/Engine/Config/Windows/Windows_SDK.json
--- a/Engine/Config/Windows/Windows_SDK.json
+++ b/Engine/Config/Windows/Windows_SDK.json
@@ -9,9 +9,9 @@
 		"20.1.8-20.999",
 		"20.1.0-20.1.7",
 		"19.1.0-19.999",
-		"18.1.3-18.999"
+		"18.1.8-18.999"
 	],
-	"MinimumClangVersion": "18.1.3",
+	"MinimumClangVersion": "18.1.8",
 
 	"//01.1": "this is a strange format, but instead of a version range, this maps min clang version for each MSVC version",
 	"//01.2": "14.44 = VS2022 17.14.x -> clang 19",
