ICX 2025.3 update

#rnx

diff --git a/Engine/Config/Windows/Windows_SDK.json b/Engine/Config/Windows/Windows_SDK.json
--- a/Engine/Config/Windows/Windows_SDK.json
+++ b/Engine/Config/Windows/Windows_SDK.json
@@ -42,9 +42,9 @@
 	"MinimumVisualCppVersion": "14.38.33130",
 
 	"PreferredIntelOneApiVersions": [
-		"2025.2.0-2025.2.9999"
+		"2025.3.0-2025.3.9999"
 	],
-	"MinimumIntelOneApiVersion": "2025.2.0",
+	"MinimumIntelOneApiVersion": "2025.3.0",
 
 	"//04.1": "Microsoft.VisualStudio.Component.VC.14.44.17.14.x86.x64 -- Only LTSC versions of MSVC should be suggested",
 	"//04.2": "Microsoft.VisualStudio.Component.VC.14.44.17.14.ATL -- Match version of MSVC",
