Added SP_PCD3D_ES3_1_NONSDF with SDF off for content validation, as default PCD3D_ES3_1 now has SDF enabled
#Synced-CL 48496786

diff --git a/Engine/Config/BaseDeviceProfiles.ini b/Engine/Config/BaseDeviceProfiles.ini
--- a/Engine/Config/BaseDeviceProfiles.ini
+++ b/Engine/Config/BaseDeviceProfiles.ini
@@ -9,7 +9,7 @@
 +DeviceProfileNameAndTypes=WindowsServer,Windows
 +DeviceProfileNameAndTypes=WindowsClient,Windows
 +DeviceProfileNameAndTypes=Windows_Preview_ES31,Windows
-+DeviceProfileNameAndTypes=Windows_Preview_ES31_SDF,Windows
++DeviceProfileNameAndTypes=Windows_Preview_ES31_NonSDF,Windows
 +DeviceProfileNameAndTypes=IOS,IOS
 +DeviceProfileNameAndTypes=iPadAir2,IOS
 +DeviceProfileNameAndTypes=iPadAir3,IOS
@@ -1516,6 +1516,6 @@ BaseProfileName=Android_Mid
 DeviceType=Windows
 BaseProfileName=Windows
 
-[Windows_Preview_ES31_SDF DeviceProfile]
+[Windows_Preview_ES31_NonSDF DeviceProfile]
 DeviceType=Windows
 BaseProfileName=Windows
diff --git a/Engine/Config/Windows/DataDrivenPlatformInfo.ini b/Engine/Config/Windows/DataDrivenPlatformInfo.ini
--- a/Engine/Config/Windows/DataDrivenPlatformInfo.ini
+++ b/Engine/Config/Windows/DataDrivenPlatformInfo.ini
@@ -169,20 +169,20 @@ MenuTooltip=LOCTEXT("PreviewMenuTooltip_PCD3D_Mobile", "Windows using Mobile pro
 DeviceProfileName=Windows_Preview_ES31
 PreviewDeviceProfileName=Windows_Preview_ES31
 
-[PreviewShaderPlatform SP_PCD3D_ES3_1_SDF]
+[PreviewShaderPlatform SP_PCD3D_ES3_1_NONSDF]
 ParentShaderPlatform=PCD3D_ES3_1
-bSupportsDistanceFields=true
+bSupportsDistanceFields=false
 bCanUsePreviewPlatformForMaterialValidation=true
-FriendlyName=LOCTEXT("FriendlyShaderPlatformName_D3D_Mobile_SDF", "PC D3D Mobile SDF")
+FriendlyName=LOCTEXT("FriendlyShaderPlatformName_D3D_Mobile_NonSDF", "PC D3D Mobile NonSDF")
 
-[PreviewPlatform PCD3D_ES3_1_SDF]
+[PreviewPlatform PCD3D_ES3_1_NONSDF]
 PlatformName=PCD3D_ES3_1
 ShaderFormat=PCD3D_ES31
-ShaderPlatform=SP_PCD3D_ES3_1_SDF
-PreviewShaderPlatform=SP_PCD3D_ES3_1_SDF
-MenuTooltip=LOCTEXT("PreviewMenuTooltip_PCD3D_Mobile_SDF", "Windows using Mobile SDF profile.")
-DeviceProfileName=Windows_Preview_ES31_SDF
-PreviewDeviceProfileName=Windows_Preview_ES31_SDF
+ShaderPlatform=SP_PCD3D_ES3_1_NONSDF
+PreviewShaderPlatform=SP_PCD3D_ES3_1_NONSDF
+MenuTooltip=LOCTEXT("PreviewMenuTooltip_PCD3D_Mobile_NONSDF", "Windows using Mobile NonSDF profile.")
+DeviceProfileName=Windows_Preview_ES31_NonSDF
+PreviewDeviceProfileName=Windows_Preview_ES31_NonSDF
 
 [ShaderPlatform OPENGL_PCES3_1]
 Language=OpenGL
