Block versions of MSVC 14.50 with known issues prior to 14.50.35723 (Visual Studio 18.2.1)

#rnx

#Synced-CL 50271719

diff --git a/Engine/Config/Windows/Windows_SDK.json b/Engine/Config/Windows/Windows_SDK.json
--- a/Engine/Config/Windows/Windows_SDK.json
+++ b/Engine/Config/Windows/Windows_SDK.json
@@ -31,10 +31,12 @@
 		"14.44.35207-14.44.99999"
 	],
 
-	"//03.1": "14.44.35207 -- VS2022 17.14.0 - 17.4.3: Template compile error. Resolved with 14.44.35211",
-	"//03.2": "14.40.x -- VS2022 17.10.0: Codegen issue causes runtime crash for AVX/2/512. 17.10.x: PGO optimize issue causes runtime crash",
-	"//03.3": "14.39.x -- VS2022 17.9.0 - 17.9.3: Internal compiler errors. 17.9.x: Codegen issue causes runtime crash for AVX/2/512",
+	"//03.1": "14.50.35717 -- VS2026 18.0.0 - 18.2.0: Internal compiler errors. Resolved with 14.50.35723",
+	"//03.2": "14.44.35207 -- VS2022 17.14.0 - 17.4.3: Template compile error. Resolved with 14.44.35211",
+	"//03.3": "14.40.x -- VS2022 17.10.0: Codegen issue causes runtime crash for AVX/2/512. 17.10.x: PGO optimize issue causes runtime crash",
+	"//03.4": "14.39.x -- VS2022 17.9.0 - 17.9.3: Internal compiler errors. 17.9.x: Codegen issue causes runtime crash for AVX/2/512",
 	"BannedVisualCppVersions": [
+		"14.50.0-14.50.35722",
 		"14.44.0-14.44.35210",
 		"14.40.0-14.43.99999",
 		"14.39.0-14.39.99999"
@@ -73,6 +75,8 @@
 	],
 
 	"VisualStudio2026SuggestedComponents": [
+		"Microsoft.VisualStudio.Component.VC.14.50.18.0.x86.x64",
+		"Microsoft.VisualStudio.Component.VC.14.50.18.0.ATL",
 	],
 
 	"MinimumVisualStudio2022Version": "17.8",
