Bump preferred MSVC to 14.50

#Synced-CL 51424418

diff --git a/Engine/Config/Windows/Windows_SDK.json b/Engine/Config/Windows/Windows_SDK.json
--- a/Engine/Config/Windows/Windows_SDK.json
+++ b/Engine/Config/Windows/Windows_SDK.json
@@ -26,8 +26,10 @@
 	],
 
 	"//02.1": "Version number is the MSVC family, which is the version in the Visual Studio folder",
-	"//02.2": "14.44.35207-14.44.99999 = VS2022 17.14.x",
+	"//02.2": "14.50.35717-14.50.99999 = VS2026 18.0",
+	"//02.3": "14.44.35207-14.44.99999 = VS2022 17.14",
 	"PreferredVisualCppVersions": [
+		"14.50.35717-14.50.99999",
 		"14.44.35207-14.44.99999"
 	],
 
@@ -48,14 +50,12 @@
 	],
 	"MinimumIntelOneApiVersion": "2025.3.0",
 
-	"//04.1": "Microsoft.VisualStudio.Component.VC.14.44.17.14.x86.x64 -- Only LTSC versions of MSVC should be suggested",
-	"//04.2": "Microsoft.VisualStudio.Component.VC.14.44.17.14.ATL -- Match version of MSVC",
+	"//04.1": "Microsoft.VisualStudio.Component.VC.14.50.18.0.x86.x64 -- Only LTSC versions of MSVC should be suggested",
+	"//04.2": "Microsoft.VisualStudio.Component.VC.14.50.18.0.ATL -- Match version of MSVC",
 	"VisualStudioSuggestedComponents": [
 		"Microsoft.Net.Component.4.6.2.TargetingPack",
-		"Microsoft.VisualStudio.Component.VC.14.44.17.14.x86.x64",
-		"Microsoft.VisualStudio.Component.VC.14.44.17.14.ATL",
-		"Microsoft.VisualStudio.Component.VC.14.38.17.8.x86.x64",
-		"Microsoft.VisualStudio.Component.VC.14.38.17.8.ATL",
+		"Microsoft.VisualStudio.Component.VC.14.50.18.0.x86.x64",
+		"Microsoft.VisualStudio.Component.VC.14.50.18.0.ATL",
 		"Microsoft.VisualStudio.Component.VC.Tools.x86.x64",
 		"Microsoft.VisualStudio.Component.Windows11SDK.22621",
 		"Microsoft.VisualStudio.Component.VC.Llvm.Clang",
diff --git a/Engine/Source/Programs/UnrealBuildTool/Platform/Windows/MicrosoftPlatformSDK.cs b/Engine/Source/Programs/UnrealBuildTool/Platform/Windows/MicrosoftPlatformSDK.cs
--- a/Engine/Source/Programs/UnrealBuildTool/Platform/Windows/MicrosoftPlatformSDK.cs
+++ b/Engine/Source/Programs/UnrealBuildTool/Platform/Windows/MicrosoftPlatformSDK.cs
@@ -624,7 +624,7 @@ namespace UnrealBuildTool
 				case WindowsCompiler.VisualStudio2022:
 					return "Visual Studio 2022";
 				case WindowsCompiler.VisualStudio2026:
-					return "Visual Studio 2026";
+					return "Visual Studio";
 				default:
 					return Compiler.ToString();
 			}
@@ -1036,10 +1036,10 @@ namespace UnrealBuildTool
 							Logger.LogDebug("Found {VisualStudio} installation: {BaseDir} (Product={ProductId}, Version={Version})", Compiler, BaseDir, ProductId, Version);
 						}
 
-						Installations = [.. Installations.OrderBy(x => x.Compiler)
-						.ThenBy(x => x.bCommunity)
-						.ThenBy(x => x.ReleaseChannel)
-						.ThenByDescending(x => x.Version)
+						Installations = [.. Installations.OrderByDescending(x => x.Compiler)
+							.ThenBy(x => x.bCommunity)
+							.ThenBy(x => x.ReleaseChannel)
+							.ThenByDescending(x => x.Version)
 						];
 					}
 					catch (Exception Ex)
diff --git a/Engine/Source/Runtime/ApplicationCore/Internal/Windows/WindowsRedistributableValidation.h b/Engine/Source/Runtime/ApplicationCore/Internal/Windows/WindowsRedistributableValidation.h
--- a/Engine/Source/Runtime/ApplicationCore/Internal/Windows/WindowsRedistributableValidation.h
+++ b/Engine/Source/Runtime/ApplicationCore/Internal/Windows/WindowsRedistributableValidation.h
@@ -37,7 +37,7 @@ struct VersionInfo
 
 // This minimum should match the version installed by
 // Engine/Extras/Redist/en-us/vc_redist.x64.exe
-inline constexpr VersionInfo MinRedistVersion = { 14, 44, 35211, 0 };
+inline constexpr VersionInfo MinRedistVersion = { 14, 50, 35719, 0 };
 
 inline constexpr VersionInfo MinGameInputRedistVersion = { 3, 1, 26100, 6879 };
 
